PCap Dev

- Performance Capture Workflow plugin move to public in Plugins/Virtual Production

#virtualized

diff --git a/Engine/Plugins/VirtualProduction/PerformanceCaptureWorkflow/Config/BasePerformanceCaptureWorkflow.ini b/Engine/Plugins/VirtualProduction/PerformanceCaptureWorkflow/Config/BasePerformanceCaptureWorkflow.ini
new file mode 100644
--- /dev/null
+++ b/Engine/Plugins/VirtualProduction/PerformanceCaptureWorkflow/Config/BasePerformanceCaptureWorkflow.ini
@@ -0,0 +1,14 @@
+[SectionsToSave]
+bCanSaveAllSections=True
+
+[/Script/PerformanceCaptureWorkflow.PerformanceCaptureSettings]
+StageRoot=/PerformanceCaptureWorkflow/Stage/BP_StageRoot.BP_StageRoot_C
+ViewModelClass=/PerformanceCaptureWorkflow/Core/VM_PCapWorkflow.VM_PCapWorkflow_C
+MocapManagerUI=/PerformanceCaptureWorkflow/Core/EW_MocapManager.EW_MocapManager
+DefaultSessionTemplate=/PerformanceCaptureWorkflow/Core/DefaultSessionTemplate.DefaultSessionTemplate
+DatabaseHelperClass=/PerformanceCaptureWorkflow/Database/EU_InternalDatabaseHelper.EU_InternalDatabaseHelper_C
+[CoreRedirects]
++ClassRedirects=(OldName="/Script/PerformanceCaptureWorkflow.PCap_DataAsset",NewName="/Script/PerformanceCaptureWorkflow.PCapDataAsset")
++ClassRedirects=(OldName="/Script/PerformanceCaptureWorkflow.PCap_CharacterDataAsset",NewName="/Script/PerformanceCaptureWorkflow.PCapCharacterDataAsset")
++ClassRedirects=(OldName="/Script/PerformanceCaptureWorkflow.PCap_PropDataAsset",NewName="/Script/PerformanceCaptureWorkflow.PCapPropDataAsset")
++ClassRedirects=(OldName="/Script/PerformanceCaptureWorkflow.PCap_PerformerDataAsset",NewName="/Script/PerformanceCaptureWorkflow.PCapPerformerDataAsset")
\ No newline at end of file
diff --git a/Engine/Plugins/VirtualProduction/PerformanceCaptureWorkflow/PerformanceCaptureWorkflow.uplugin b/Engine/Plugins/VirtualProduction/PerformanceCaptureWorkflow/PerformanceCaptureWorkflow.uplugin
new file mode 100644
--- /dev/null
+++ b/Engine/Plugins/VirtualProduction/PerformanceCaptureWorkflow/PerformanceCaptureWorkflow.uplugin
@@ -0,0 +1,70 @@
+{
+	"FileVersion": 3,
+	"Version": 1,
+	"VersionName": "1.0",
+	"FriendlyName": "Performance Capture Workflow",
+	"Description": "Performance Capture In-Editor Workflow tools.",
+	"Category": "Animation",
+	"CreatedBy": "Epic Games",
+	"CreatedByURL": "",
+	"DocsURL": "",
+	"MarketplaceURL": "",
+	"SupportURL": "",
+	"CanContainContent": true,
+	"IsBetaVersion": false,
+	"IsExperimentalVersion": true,
+	"Installed": false,
+	"Modules": [
+		{
+			"Name": "PerformanceCaptureWorkflow",
+			"Type": "Editor",
+			"LoadingPhase": "Default"
+		}
+	],
+	"Plugins": [
+		{
+			"Name": "LiveLink",
+			"Enabled": true
+		},
+		{
+			"Name": "PerformanceCaptureCore",
+			"Enabled": true
+		},
+		{
+			"Name": "GeometryScripting",
+			"Enabled": true
+		},
+		{
+			"Name": "CommonUI",
+			"Enabled": true
+		},
+		{
+			"Name": "SkeletalMeshModelingTools",
+			"Enabled": true
+		},
+		{
+			"Name": "EditorScriptingUtilities",
+			"Enabled": true
+		},
+		{
+			"Name": "ModelViewViewModel",
+			"Enabled": true
+		},
+		{
+			"Name": "EngineAssetDefinitions",
+			"Enabled": true
+		},
+		{
+			"Name": "Takes",
+			"Enabled": true
+		},
+		{
+			"Name": "DirectoryPlaceholder",
+			"Enabled": true
+		},
+		{
+			"Name": "NamingTokens",
+			"Enabled": true
+		}
+	]
+}
\ No newline at end of file
diff --git a/Engine/Plugins/VirtualProduction/PerformanceCaptureWorkflow/Source/PerformanceCaptureWorkflow/PerformanceCaptureWorkflow.Build.cs b/Engine/Plugins/VirtualProduction/PerformanceCaptureWorkflow/Source/PerformanceCaptureWorkflow/PerformanceCaptureWorkflow.Build.cs
new file mode 100644
--- /dev/null
+++ b/Engine/Plugins/VirtualProduction/PerformanceCaptureWorkflow/Source/PerformanceCaptureWorkflow/PerformanceCaptureWorkflow.Build.cs
@@ -0,0 +1,58 @@
+// Copyright Epic Games, Inc. All Rights Reserved.
+
+using UnrealBuildTool;
+
+public class PerformanceCaptureWorkflow : ModuleRules
+{
+	public PerformanceCaptureWorkflow(ReadOnlyTargetRules Target) : base(Target)
+	{
+		PCHUsage = ModuleRules.PCHUsageMode.UseExplicitOrSharedPCHs;
+		ShortName = "PCWF";
+
+		PublicDependencyModuleNames.AddRange(
+			new string[]
+			{
+				"Core",
+				"CommonUI",
+				"PlacementMode"
+			}
+			);
+		
+		PrivateDependencyModuleNames.AddRange(
+			new string[]
+			{
+				"CoreUObject",
+				"Engine",
+				"Slate",
+				"SlateCore",
+				"DeveloperSettings",
+				"Settings",
+				"PerformanceCaptureCore",
+				"LiveLink",
+				"LiveLinkInterface",
+				"IKRig", 
+				"LiveLinkAnimationCore",
+				"LevelSequence",
+				"Projects",
+				"InputCore",
+				"EditorFramework",
+				"UnrealEd",
+				"ToolMenus",
+				"ModelViewViewModel",
+				"UMG",
+				"UMGEditor",
+				"Blutility",
+				"EditorSubsystem",
+				"AnimationCore", 
+				"DataTableEditor",
+				"AssetDefinition",
+				"EngineAssetDefinitions",
+				"CinematicCamera",
+				"TakeRecorderSources", 
+				"TakesCore",
+				"SubobjectDataInterface",
+				"NamingTokens"
+			}
+			);
+	}
+}
diff --git a/Engine/Plugins/VirtualProduction/PerformanceCaptureWorkflow/Source/PerformanceCaptureWorkflow/Private/AnimNode_LiveLinkProp.cpp b/Engine/Plugins/VirtualProduction/PerformanceCaptureWorkflow/Source/PerformanceCaptureWorkflow/Private/AnimNode_LiveLinkProp.cpp
new file mode 100644
--- /dev/null
+++ b/Engine/Plugins/VirtualProduction/PerformanceCaptureWorkflow/Source/PerformanceCaptureWorkflow/Private/AnimNode_LiveLinkProp.cpp
@@ -0,0 +1,140 @@
+// Copyright Epic Games, Inc. All Rights Reserved.
+
+
+#include "AnimNode_LiveLinkProp.h"
+
+#include "Animation/AnimInstanceProxy.h"
+#include "Animation/AnimTrace.h"
+#include "Features/IModularFeatures.h"
+#include "ILiveLinkClient.h"
+#include "Roles/LiveLinkAnimationRole.h"
+#include "Roles/LiveLinkAnimationTypes.h"
+
+#include UE_INLINE_GENERATED_CPP_BY_NAME(AnimNode_LiveLinkProp)
+
+FAnimNode_LiveLinkProp::FAnimNode_LiveLinkProp()
+	: LiveLinkClient_AnyThread(nullptr)
+	, CachedDeltaTime(0.0f)
+{
+}
+
+void FAnimNode_LiveLinkProp::BuildPoseFromAnimData(const FLiveLinkSubjectFrameData& LiveLinkData, FPoseContext& Output)
+{
+	const FLiveLinkSkeletonStaticData* SkeletonData = LiveLinkData.StaticData.Cast<FLiveLinkSkeletonStaticData>();
+	const FLiveLinkAnimationFrameData* FrameData = LiveLinkData.FrameData.Cast<FLiveLinkAnimationFrameData>();
+	check(SkeletonData);
+	check(FrameData);
+
+	//Get bones names
+	const TArray<FName>& SourceBoneNames = SkeletonData->GetBoneNames();
+
+	//Iterate through bones in LiveLink Data
+	for (int32 i = 0; i < SourceBoneNames.Num(); ++i)
+	{
+		FName BoneName = SourceBoneNames[i];
+
+		FTransform BoneTransform = OffsetTransform*FrameData->Transforms[i];
+
+		int32 MeshIndex = Output.Pose.GetBoneContainer().GetPoseBoneIndexForBoneName(BoneName);
+		if (MeshIndex != INDEX_NONE)
+		{
+			FCompactPoseBoneIndex CompactPoseBoneIndex = Output.Pose.GetBoneContainer().MakeCompactPoseIndex(FMeshPoseBoneIndex(MeshIndex));
+			if(CompactPoseBoneIndex != INDEX_NONE)
+			{
+				Output.Pose[CompactPoseBoneIndex] = BoneTransform;
+			}
+		}
+	}
+	CachedDeltaTime = 0.0f;
+}
+
+void FAnimNode_LiveLinkProp::Initialize_AnyThread(const FAnimationInitializeContext& Context)
+{
+	InputPose.Initialize(Context);
+
+	if (CachedLiveLinkFrameData.IsValid() == false)
+	{
+		CachedLiveLinkFrameData  = MakeShared<FLiveLinkSubjectFrameData>();
+	}
+}
+
+void FAnimNode_LiveLinkProp::PreUpdate(const UAnimInstance* InAnimInstance)
+{
+	ILiveLinkClient* ThisFrameClient = nullptr;
+	IModularFeatures& ModularFeatures  = IModularFeatures::Get();
+	if (ModularFeatures.IsModularFeatureAvailable(ILiveLinkClient::ModularFeatureName))
+	{
+		ThisFrameClient = &IModularFeatures::Get().GetModularFeature<ILiveLinkClient>(ILiveLinkClient::ModularFeatureName);
+	}
+	LiveLinkClient_AnyThread = ThisFrameClient;
+}
+
+void FAnimNode_LiveLinkProp::Update_AnyThread(const FAnimationUpdateContext& Context)
+{
+	InputPose.Update(Context);
+
+	GetEvaluateGraphExposedInputs().Execute(Context);
+
+	CachedDeltaTime += Context.GetDeltaTime();
+
+	TRACE_ANIM_NODE_VALUE(Context, TEXT("SubjectName"), LiveLinkSubjectName.Name);
+}
+
+void FAnimNode_LiveLinkProp::Evaluate_AnyThread(FPoseContext& Output)
+{
+	InputPose.Evaluate(Output);
+
+	if (!LiveLinkClient_AnyThread)
+	{
+		return;
+	}
+
+	FLiveLinkSubjectFrameData SubjectFrameData;
+
+	if(bDoLiveLinkEvaluation)
+	{
+		// Invalidate cached evaluated Role to make sure we have a valid one during the last evaluation when using it
+		CachedEvaluatedRole = nullptr;
+
+		TSubclassOf<ULiveLinkRole> SubjectRole = LiveLinkClient_AnyThread->GetSubjectRole_AnyThread(LiveLinkSubjectName);
+		if (SubjectRole)
+		{
+			if(LiveLinkClient_AnyThread->DoesSubjectSupportsRole_AnyThread(LiveLinkSubjectName, ULiveLinkAnimationRole::StaticClass()))
+			{
+				//Process the animation data
+				if(LiveLinkClient_AnyThread->EvaluateFrame_AnyThread(LiveLinkSubjectName, ULiveLinkAnimationRole::StaticClass(), SubjectFrameData))
+				{
+					BuildPoseFromAnimData(SubjectFrameData, Output);
+					CachedEvaluatedRole = ULiveLinkAnimationRole::StaticClass();
+				}
+			}
+			CachedLiveLinkFrameData->StaticData = MoveTemp(SubjectFrameData.StaticData);
+			CachedLiveLinkFrameData->FrameData = MoveTemp(SubjectFrameData.FrameData);
+		}
+	}
+	else
+	{
+		if(CachedLiveLinkFrameData && CachedEvaluatedRole)
+		{
+			if (CachedEvaluatedRole == ULiveLinkAnimationRole::StaticClass())
+			{
+				BuildPoseFromAnimData(*CachedLiveLinkFrameData, Output);
+			}
+		}
+	}
+}
+
+void FAnimNode_LiveLinkProp::CacheBones_AnyThread(const FAnimationCacheBonesContext& Context)
+{
+	Super::CacheBones_AnyThread(Context);
+	InputPose.CacheBones(Context);
+}
+
+void FAnimNode_LiveLinkProp::GatherDebugData(FNodeDebugData& DebugData)
+{
+	FString DebugLine = FString::Printf(TEXT("LiveLink Prop - SubjectName: %s"), *LiveLinkSubjectName.ToString());
+
+	DebugData.AddDebugItem(DebugLine);
+	InputPose.GatherDebugData(DebugData);
+}
+
diff --git a/Engine/Plugins/VirtualProduction/PerformanceCaptureWorkflow/Source/PerformanceCaptureWorkflow/Private/AnimNode_LiveLinkProp.h b/Engine/Plugins/VirtualProduction/PerformanceCaptureWorkflow/Source/PerformanceCaptureWorkflow/Private/AnimNode_LiveLinkProp.h
new file mode 100644
--- /dev/null
+++ b/Engine/Plugins/VirtualProduction/PerformanceCaptureWorkflow/Source/PerformanceCaptureWorkflow/Private/AnimNode_LiveLinkProp.h
@@ -0,0 +1,67 @@
+// Copyright Epic Games, Inc. All Rights Reserved.
+
+#pragma once
+
+#include "Animation/AnimNodeBase.h"
+
+#include "CoreMinimal.h"
+#include "LiveLinkTypes.h"
+#include "Templates/SubclassOf.h"
+
+#include "AnimNode_LiveLinkProp.generated.h"
+
+class ILiveLinkClient;
+class ULiveLinkRole;
+struct FLiveLinkSubjectFrameData;
+
+/**
+ * This animnode is exclusively for Mocap props - single bone skeleton meshes. Not exposed to the animation graph
+ */
+
+USTRUCT(BlueprintInternalUseOnly)
+struct FAnimNode_LiveLinkProp : public FAnimNode_Base
+{
+	GENERATED_BODY()
+
+public:
+	UPROPERTY(EditAnywhere, BlueprintReadWrite, Category = "Performance Capture|Prop")
+	FPoseLink InputPose;
+
+	UPROPERTY(EditAnywhere, BlueprintReadWrite, Category = "Performance Capture|Prop")
+	FLiveLinkSubjectName LiveLinkSubjectName;
+
+	UPROPERTY(EditAnywhere, BlueprintReadWrite, Transient, Category = "Performance Capture|Prop")
+	bool bDoLiveLinkEvaluation = true;
+
+	UPROPERTY(EditAnywhere, BlueprintReadWrite, Category = "Performance Capture|Prop")
+	FTransform OffsetTransform = FTransform::Identity;
+
+	
+
+	FAnimNode_LiveLinkProp();
+
+	//~ FAnimNode_Base interface
+	virtual void Initialize_AnyThread(const FAnimationInitializeContext& Context) override;
+	virtual void PreUpdate(const UAnimInstance* InAnimInstance) override;
+	virtual void Update_AnyThread(const FAnimationUpdateContext & Context) override;
+	virtual void Evaluate_AnyThread(FPoseContext& Output) override;
+	virtual void CacheBones_AnyThread(const FAnimationCacheBonesContext & Context) override;
+	virtual void GatherDebugData(FNodeDebugData& DebugData) override;
+	virtual bool HasPreUpdate() const override { return true; };
+	//~ End of FAnimNode_Base interface
+
+protected:
+	
+	void BuildPoseFromAnimData(const FLiveLinkSubjectFrameData& LiveLinkData, FPoseContext& Output);
+
+private:
+
+	ILiveLinkClient* LiveLinkClient_AnyThread;
+
+	//Delta time - TODO: might not need this because there is no retarget class at play
+	float CachedDeltaTime;
+
+	TSharedPtr<FLiveLinkSubjectFrameData> CachedLiveLinkFrameData;
+
+	TSubclassOf<class ULiveLinkRole> CachedEvaluatedRole;
+};
diff --git a/Engine/Plugins/VirtualProduction/PerformanceCaptureWorkflow/Source/PerformanceCaptureWorkflow/Private/PCapActorFactory.cpp b/Engine/Plugins/VirtualProduction/PerformanceCaptureWorkflow/Source/PerformanceCaptureWorkflow/Private/PCapActorFactory.cpp
new file mode 100644
--- /dev/null
+++ b/Engine/Plugins/VirtualProduction/PerformanceCaptureWorkflow/Source/PerformanceCaptureWorkflow/Private/PCapActorFactory.cpp
@@ -0,0 +1,364 @@
+// Copyright Epic Games, Inc. All Rights Reserved.
+
+#include "PCapActorFactory.h"
+
+#include "PCapDatabase.h"
+#include "Retargeter/IKRetargeter.h"
+#include "Engine/SkeletalMesh.h"
+#include "Components/SkeletalMeshComponent.h"
+#include "Engine/StaticMeshActor.h"
+#include "Engine/StaticMesh.h"
+#include "Materials/MaterialInstanceDynamic.h"
+#include "AssetSelection.h"
+#include "PCapPropComponent.h"
+#include "PerformanceCapture.h"
+#include "SubobjectDataSubsystem.h"
+#include "Elements/Actor/ActorElementData.h"
+#include "Elements/Framework/EngineElementsLibrary.h"
+#include "Subsystems/EditorActorSubsystem.h"
+#include "Subsystems/PlacementSubsystem.h"
+
+
+/*------------------------------------------------------------------------------
+	Custom Actor Factory for Character Data Asset implementation.
+------------------------------------------------------------------------------*/
+UClass* UPCapCharacterActorFactory::GetDefaultActorClass(const FAssetData& AssetData)
+{
+	if(AssetData.IsValid())
+	{
+		if(UPCapCharacterDataAsset* CharacterDataAsset = Cast<UPCapCharacterDataAsset>(AssetData.GetAsset()))
+		{
+			if(UClass* Class = CharacterDataAsset->CaptureCharacterClass.LoadSynchronous())
+			{
+				return Class;
+			}	
+		}
+	}
+	return nullptr;
+}
+
+bool UPCapCharacterActorFactory::CanCreateActorFrom(const FAssetData& AssetData, FText& OutErrorMsg)
+{
+	if(AssetData.IsValid())
+	{
+		if(UPCapCharacterDataAsset* CharacterDataAsset = Cast<UPCapCharacterDataAsset>(AssetData.GetAsset()))
+		{
+			const UClass* Class = CharacterDataAsset->CaptureCharacterClass.LoadSynchronous();
+			const USkeletalMesh* SkeletalMesh = CharacterDataAsset->SkeletalMesh.LoadSynchronous();
+
+			if(Class && SkeletalMesh)
+			{
+				return true;
+			}	
+		}
+	}
+	OutErrorMsg = NSLOCTEXT("Performance Capture", "CanCreateActorFrom_PerformerAsset", "Asset does not contain valid Capture Character Class or Skeletal Mesh");
+	return false;
+}
+
+void UPCapCharacterActorFactory::PostSpawnActor(UObject* Asset, AActor* NewActor)
+{
+	Super::PostSpawnActor(Asset, NewActor);
+	
+	if(Cast<ACaptureCharacter>(NewActor))
+	{
+		ACaptureCharacter* NewCharacter = Cast<ACaptureCharacter>(NewActor);
+		
+		if(UPCapCharacterDataAsset* CharacterAsset = Cast<UPCapCharacterDataAsset>(Asset))
+		{
+			USkeletalMesh* CharacterMesh = CharacterAsset->SkeletalMesh.LoadSynchronous();
+			NewCharacter->GetSkeletalMeshComponent()->SetSkeletalMeshAsset(CharacterMesh);
+			NewCharacter->SetRetargetAsset(Cast<UIKRetargeter>(CharacterAsset->Retargeter.LoadSynchronous()));
+			
+			if(SpawnedPerformerHandle.IsSet()) //Entry here is only once per placing action.
+			{
+				ACapturePerformer* SourcePerformer = Cast<ACapturePerformer>(ActorElementDataUtil::GetActorFromHandle(SpawnedPerformerHandle));
+				NewCharacter->SourcePerformer = SourcePerformer;
+				NewCharacter->GetSkeletalMeshComponent()->SetReceivesDecals(false);
+
+				if(!CharacterAsset->AdditionalMeshes.IsEmpty())
+				{
+					if (NewActor->HasAnyFlags(RF_Transactional) == true) //Do not spawn any subcomponents on the transient copy of the new actor
+					{
+						USubobjectDataSubsystem* SubSystem = GEngine->GetEngineSubsystem<USubobjectDataSubsystem>();
+
+						for (TSoftObjectPtr<USkeletalMesh> Mesh : CharacterAsset->AdditionalMeshes)
+						{
+							TArray<FSubobjectDataHandle> SubObjectHandles;
+							SubSystem->GatherSubobjectData(NewActor, SubObjectHandles);
+							FSubobjectDataHandle RootHandle = SubObjectHandles[0]; //The actor's root component is always the first handle in the array
+							FAddNewSubobjectParams NewObjectParams;
+							NewObjectParams.ParentHandle = RootHandle;
+							NewObjectParams.NewClass = USkeletalMeshComponent::StaticClass();
+							NewObjectParams.bConformTransformToParent = true;
+
+							FText FailureReason;
+							FSubobjectDataHandle NewSkelMeshComponentHandle = SubSystem->AddNewSubobject(NewObjectParams, FailureReason);
+					
+							TArray<USkeletalMeshComponent*> SkeletalMeshComponents;
+							NewCharacter->GetComponents(SkeletalMeshComponents);
+					
+							if(TObjectPtr<USkeletalMeshComponent> const NewSkeletalMeshComponent = SkeletalMeshComponents.Last())
+							{
+								NewSkeletalMeshComponent->SetSkeletalMeshAsset(Mesh.LoadSynchronous());
+								NewSkeletalMeshComponent->SetReceivesDecals(false);
+							}
+						}
+					}
+				}
+			}
+		}
+	}
+	SpawnedPerformerHandle.Release();
+}
+
+TArray<FTypedElementHandle> UPCapCharacterActorFactory::PlaceAsset(const FAssetPlacementInfo& InPlacementInfo, const FPlacementOptions& InPlacementOptions)
+{
+	TArray<FTypedElementHandle> PlacedActorHandles;
+	FActorSpawnParameters SpawnParams;
+	SpawnParams.Name = InPlacementInfo.NameOverride;
+	SpawnParams.ObjectFlags = InPlacementOptions.bIsCreatingPreviewElements ? EObjectFlags::RF_Transient : EObjectFlags::RF_Transactional;
+	SpawnParams.bTemporaryEditorActor = InPlacementOptions.bIsCreatingPreviewElements;
+
+	AActor* NewActor = CreateActor(InPlacementInfo.AssetToPlace.GetAsset(), InPlacementInfo.PreferredLevel.Get(), InPlacementInfo.FinalizedTransform, SpawnParams);
+	if (NewActor)
+	{
+		PlacedActorHandles.Add(UEngineElementsLibrary::AcquireEditorActorElementHandle(NewActor));
+
+		// Run post placement steps
+		if (InPlacementOptions.bIsCreatingPreviewElements)
+		{
+			NewActor->SetActorEnableCollision(false);
+		}
+	}
+
+	//Get Performer Asset and Attempt to spawn that too. 
+	const UPCapCharacterDataAsset* CharacterData = Cast<UPCapCharacterDataAsset>(InPlacementInfo.AssetToPlace.GetAsset());
+	
+	if(	UPCapPerformerDataAsset* PerformerDataAsset = CharacterData->SourcePerformerAsset.LoadSynchronous()) //If the Character Asset has a valid performer, place/spawn 
+	{
+		UE::AssetPlacementUtil::FExtraPlaceAssetOptions PlacementOptions;
+			TArray<FTypedElementHandle> NewPerformerHandles = PlaceAssetInCurrentLevel(PerformerDataAsset, PlacementOptions);
+		PlacedActorHandles.Append(NewPerformerHandles);
+		if(!NewPerformerHandles.IsEmpty())
+		{
+			//If successfully spawned, cache the created handle so we have a reference to the created Performer to use in Post operations.
+			SpawnedPerformerHandle = NewPerformerHandles[0];
+		}
+	}
+	return PlacedActorHandles;
+}
+
+/*------------------------------------------------------------------------------
+	Custom Actor Factory for Performer Data Asset implementation.
+------------------------------------------------------------------------------*/
+UClass* UPCapPerformerActorFactory::GetDefaultActorClass(const FAssetData& AssetData)
+{
+	if(AssetData.IsValid())
+	{
+		if (UPCapPerformerDataAsset* PerformerDataAsset = Cast<UPCapPerformerDataAsset>(AssetData.GetAsset()))
+		{
+			if(UClass* Class = PerformerDataAsset->PerformerActorClass.LoadSynchronous())
+			{
+				return Class;
+			}
+		}
+	}
+	return nullptr;
+}
+
+bool UPCapPerformerActorFactory::CanCreateActorFrom(const FAssetData& AssetData, FText& OutErrorMsg)
+{
+	if(AssetData.IsValid())
+	{
+		if(UPCapPerformerDataAsset* PerformerDataAsset = Cast<UPCapPerformerDataAsset>(AssetData.GetAsset()))
+		{
+			const UClass* Class = PerformerDataAsset->PerformerActorClass.LoadSynchronous();
+			const USkeletalMesh* BaseSkeletalMesh = PerformerDataAsset->BaseSkeletalMesh.LoadSynchronous();
+			const USkeletalMesh* ProportionedMesh = PerformerDataAsset->PerformerProportionedMesh.LoadSynchronous();
+			if(Class && (BaseSkeletalMesh || ProportionedMesh))
+			{
+					return true;
+			}	
+		}
+	}
+	OutErrorMsg = NSLOCTEXT("Performance Capture", "CanCreateActorFrom_Asset", "Asset is missing a a valid Capture Performer Class or Skeletal Meshes");
+	return false;
+}
+
+void UPCapPerformerActorFactory::PostSpawnActor(UObject* Asset, AActor* NewActor)
+{
+	Super::PostSpawnActor(Asset, NewActor);
+
+	if(ACapturePerformer* NewPerformer = Cast<ACapturePerformer>(NewActor))
+	{
+		UPCapPerformerDataAsset* PerformerDataAsset = Cast<UPCapPerformerDataAsset>(Asset);
+		if(PerformerDataAsset)
+		{
+			if(USkeletalMesh* PerformerScaledMesh = PerformerDataAsset->PerformerProportionedMesh.LoadSynchronous())
+			{
+				NewPerformer->SetMocapMesh(PerformerScaledMesh);
+			}
+			if(!PerformerDataAsset->PerformerProportionedMesh.IsValid() && PerformerDataAsset->BaseSkeletalMesh.IsValid()) //If there is not ProportionedMesh but a valid Base Mesh, use that instead
+			{
+				USkeletalMesh* PerformerBaseMesh = PerformerDataAsset->BaseSkeletalMesh.LoadSynchronous();
+				NewPerformer->SetMocapMesh(PerformerBaseMesh);
+			}
+
+			NewPerformer->SetLiveLinkSubject(PerformerDataAsset->LiveLinkSubject);
+			NewActor->SetActorLabel(PerformerDataAsset->PerformerName.ToString(), false);
+			
+			if(USkeletalMeshComponent* PerformerSkelMeshComponent = NewPerformer->GetSkeletalMeshComponent())
+			{
+				PerformerSkelMeshComponent->SetReceivesDecals(false);
+				
+				if(UMaterialInterface* SkelMeshMaterial = NewPerformer->GetSkeletalMeshComponent()->GetMaterial(0))
+				{
+					UMaterialInstanceDynamic* DynamicMaterial = PerformerSkelMeshComponent->CreateDynamicMaterialInstance(
+						0, SkelMeshMaterial);
+					PerformerSkelMeshComponent->SetMaterial(0, DynamicMaterial);
+					DynamicMaterial->ClearParameterValues();
+					DynamicMaterial->SetVectorParameterValue(FName(TEXT("PerformerColor")), PerformerDataAsset->PerformerColor);
+				}
+			}
+		}
+	}
+}
+
+/*------------------------------------------------------------------------------
+	Custom Actor Factory for Prop Data Asset implementation.
+------------------------------------------------------------------------------*/
+UClass* UPCapPropActorFactory::GetDefaultActorClass(const FAssetData& AssetData)
+{
+	if(AssetData.IsValid())
+	{
+		if(UPCapPropDataAsset* PropDataAsset = Cast<UPCapPropDataAsset>(AssetData.GetAsset()))
+		{
+			if(UClass* CustomClass =  PropDataAsset->CustomPropClass.Get())
+			{
+				return CustomClass;
+			}
+			
+			if(PropDataAsset->PropSkeletalMesh.IsValid())
+			{
+				if(UClass* SkeletalMeshClass = ASkeletalMeshActor::StaticClass())
+				{
+					return SkeletalMeshClass;
+				}
+			}
+
+			if(PropDataAsset->PropStaticMesh.IsValid())
+			{
+				if(UClass* StaticMeshClass = AStaticMeshActor::StaticClass())
+				{
+					return StaticMeshClass;
+				}
+			}
+		}
+	}
+	return nullptr;
+}
+
+bool UPCapPropActorFactory::CanCreateActorFrom(const FAssetData& AssetData, FText& OutErrorMsg)
+{
+	if(AssetData.IsValid())
+	{
+		if(UPCapPropDataAsset* PropDataAsset = Cast<UPCapPropDataAsset>(AssetData.GetAsset()))
+		{
+			const UStaticMesh* StaticMesh = PropDataAsset->PropStaticMesh.LoadSynchronous();
+			const USkeletalMesh* SkeletalMesh = PropDataAsset->PropSkeletalMesh.LoadSynchronous();;
+			const UClass* CustomProp = PropDataAsset->CustomPropClass.LoadSynchronous();
+						
+			if(SkeletalMesh || StaticMesh || CustomProp)
+			{
+				return true;
+			}
+		}
+	}
+	OutErrorMsg = NSLOCTEXT("Performance Capture", "CanCreateActorFrom_PropAsset", "Asset does not contain valid Static Mesh or Skeletal Mesh");
+	return false;
+}
+
+void UPCapPropActorFactory::PostSpawnActor(UObject* Asset, AActor* NewActor)
+{
+	Super::PostSpawnActor(Asset, NewActor);	
+	
+	if(UPCapPropDataAsset* PropDataAsset = Cast<UPCapPropDataAsset>(Asset))
+	{
+		if(AStaticMeshActor* NewStaticMeshProp = Cast<AStaticMeshActor>(NewActor))
+		{		
+			if(UStaticMesh* StaticMesh = PropDataAsset->PropStaticMesh.LoadSynchronous())
+			{
+				NewStaticMeshProp->GetStaticMeshComponent()->SetStaticMesh(PropDataAsset->PropStaticMesh.LoadSynchronous());
+				NewStaticMeshProp->GetStaticMeshComponent()->SetReceivesDecals(false);
+				if(PropDataAsset->PropName.IsValid()) //If the user has filled out a valid FName use that for the actor label
+				{
+					NewStaticMeshProp->SetActorLabel(PropDataAsset->PropName.ToString(), false);
+				}
+				else //Otherwise use the name of the DataAsset
+				{
+					NewStaticMeshProp->SetActorLabel(PropDataAsset->GetName(), false);
+				}
+			}
+			NewStaticMeshProp->SetActorHiddenInGame(PropDataAsset->bHiddenInGame);
+		}
+		
+		if(ASkeletalMeshActor* NewSkeletalMeshProp = Cast<ASkeletalMeshActor>(NewActor))
+		{
+			if(USkeletalMesh* SkeletalMesh = PropDataAsset->PropSkeletalMesh.LoadSynchronous())
+			{
+				NewSkeletalMeshProp->GetSkeletalMeshComponent()->SetSkeletalMeshAsset(PropDataAsset->PropSkeletalMesh.LoadSynchronous());
+				NewSkeletalMeshProp->GetSkeletalMeshComponent()->SetReceivesDecals(false);
+				if(PropDataAsset->PropName.IsValid()) //If the user has filled out a valid FName use that for the actor label
+				{
+					NewSkeletalMeshProp->SetActorLabel(PropDataAsset->PropName.ToString(), false);
+				}
+				else //Otherwise use the name of the DataAsset
+				{
+					NewSkeletalMeshProp->SetActorLabel(PropDataAsset->GetName(), false);
+				}
+				
+			}
+			NewSkeletalMeshProp->SetActorHiddenInGame(PropDataAsset->bHiddenInGame);
+		}
+		//Only add Subcomponents for the non-transient actor because transient means it is still being placed. Also, don't add the prop component for custom classes
+		if(NewActor->HasAnyFlags(RF_Transient) != true && NewActor->GetClass()!= (PropDataAsset->CustomPropClass.Get())) 
+		{
+			USubobjectDataSubsystem* SubSystem = GEngine->GetEngineSubsystem<USubobjectDataSubsystem>();
+			TObjectPtr<UPCapPropComponent> NewPropComponent =  Cast<UPCapPropComponent>(NewActor->GetComponentByClass(UPCapPropComponent::StaticClass()));
+			if(!NewPropComponent)
+			{
+				TArray<FSubobjectDataHandle> Handles;
+				SubSystem->GatherSubobjectData(NewActor, Handles);
+
+				FSubobjectDataHandle Handle = Handles[0]; //Root Component is always first entry in handles array
+
+				FAddNewSubobjectParams NewObjectParams;
+
+				NewObjectParams.ParentHandle = Handle;
+				NewObjectParams.NewClass = UPCapPropComponent::StaticClass();
+				NewObjectParams.bConformTransformToParent = true;
+
+				FText FailureReason;
+				FSubobjectDataHandle NewPropComponentHandle = SubSystem->AddNewSubobject(NewObjectParams, FailureReason);
+			}
+			
+			if(NewPropComponent)
+			{
+				NewPropComponent->SetLiveLinkSubject(PropDataAsset->LiveLinkSubject);
+				NewPropComponent->SetOffsetTransform(PropDataAsset->PropOffsetTransform);
+				if (Asset) 
+				{
+					if (Asset->IsAsset())
+					{
+						FAssetData NewAssetData = FAssetData(Asset);
+						FName SpawningPackageName = NewAssetData.PackageName;
+						NewPropComponent->SpawningDataAsset = SpawningPackageName; //Add the package name of the spawning data asset to the prop component so the toolset can locate the data asset for a given prop
+					}
+				}
+			}
+			UE_LOG(LogPCap, Display, TEXT("New PCap Prop Component added to %s"), *NewActor->GetActorLabel())
+		}
+	}
+}
+
diff --git a/Engine/Plugins/VirtualProduction/PerformanceCaptureWorkflow/Source/PerformanceCaptureWorkflow/Private/PCapActorFactory.h b/Engine/Plugins/VirtualProduction/PerformanceCaptureWorkflow/Source/PerformanceCaptureWorkflow/Private/PCapActorFactory.h
new file mode 100644
--- /dev/null
+++ b/Engine/Plugins/VirtualProduction/PerformanceCaptureWorkflow/Source/PerformanceCaptureWorkflow/Private/PCapActorFactory.h
@@ -0,0 +1,59 @@
+// Copyright Epic Games, Inc. All Rights Reserved.
+
+#pragma once
+
+#include "CoreMinimal.h"
+#include "ActorFactories/ActorFactory.h"
+#include "Elements/Framework/TypedElementHandle.h"
+#include "PCapActorFactory.generated.h"
+
+class ACapturePerformer;
+struct FTypedElementHandle;
+
+/**
+ * 
+ */
+
+UCLASS()
+class UPCapCharacterActorFactory : public UActorFactory
+{
+	GENERATED_BODY()
+
+	FTypedElementHandle SpawnedPerformerHandle;
+	
+	/** Initialize NewActorClass if necessary, and return that class. */
+	virtual UClass* GetDefaultActorClass( const FAssetData& AssetData ) override;
+
+	virtual bool CanCreateActorFrom( const FAssetData& AssetData, FText& OutErrorMsg ) override;
+
+	virtual void PostSpawnActor( UObject* Asset, AActor* NewActor ) override;
+
+	virtual TArray<FTypedElementHandle> PlaceAsset(const FAssetPlacementInfo& InPlacementInfo, const FPlacementOptions& InPlacementOptions) override;
+};
+
+UCLASS()
+class UPCapPerformerActorFactory : public UActorFactory
+{
+	GENERATED_BODY()
+
+	/** Initialize NewActorClass if necessary, and return that class. */
+	virtual UClass* GetDefaultActorClass( const FAssetData& AssetData ) override;
+
+	virtual bool CanCreateActorFrom( const FAssetData& AssetData, FText& OutErrorMsg ) override;
+
+	virtual void PostSpawnActor( UObject* Asset, AActor* NewActor ) override;
+};
+
+UCLASS()
+
+class UPCapPropActorFactory : public UActorFactory
+{
+	GENERATED_BODY()
+
+	/** Initialize NewActorClass if necessary, and return that class. */
+	virtual UClass* GetDefaultActorClass(const FAssetData& AssetData) override;
+
+	virtual bool CanCreateActorFrom(const FAssetData& AssetData, FText& OutErrorMsg) override;
+
+	virtual void PostSpawnActor( UObject* Asset, AActor* NewActor ) override;
+};
\ No newline at end of file
diff --git a/Engine/Plugins/VirtualProduction/PerformanceCaptureWorkflow/Source/PerformanceCaptureWorkflow/Private/PCapAssetDefinition.cpp b/Engine/Plugins/VirtualProduction/PerformanceCaptureWorkflow/Source/PerformanceCaptureWorkflow/Private/PCapAssetDefinition.cpp
new file mode 100644
--- /dev/null
+++ b/Engine/Plugins/VirtualProduction/PerformanceCaptureWorkflow/Source/PerformanceCaptureWorkflow/Private/PCapAssetDefinition.cpp
@@ -0,0 +1,66 @@
+// Copyright Epic Games, Inc. All Rights Reserved.
+
+
+#include "PCapAssetDefinition.h"
+#include "PerformanceCaptureStyle.h"
+#include "Misc/FileHelper.h"
+#include "Misc/MessageDialog.h"
+#include "Settings/EditorLoadingSavingSettings.h"
+
+#define LOCTEXT_NAMESPACE "UAssetDefinition_PerformanceCapture"
+
+/*------------------------------------------------------------------------------
+	UPCapDataAsset Performer implementation.
+------------------------------------------------------------------------------*/
+
+const FSlateBrush* UAssetDefinition_PerformerDataAsset::GetIconBrush(const FAssetData& InAssetData, const FName InClassName) const
+{
+	return FPerformanceCaptureStyle::Get().GetBrush("ClassIcon.PCapPerformerDataAsset");
+}
+
+const FSlateBrush* UAssetDefinition_PerformerDataAsset::GetThumbnailBrush(const FAssetData& InAssetData, const FName InClassName) const
+{
+	return FPerformanceCaptureStyle::Get().GetBrush("ClassThumbnail.PCapPerformerDataAsset");
+}
+
+/*------------------------------------------------------------------------------
+	UPCapDataAsset Character implementation.
+------------------------------------------------------------------------------*/
+
+const FSlateBrush* UAssetDefinition_CharacterDataAsset::GetIconBrush(const FAssetData& InAssetData, const FName InClassName) const
+{
+	return FPerformanceCaptureStyle::Get().GetBrush("ClassIcon.PCapCharacterDataAsset");
+}
+
+const FSlateBrush* UAssetDefinition_CharacterDataAsset::GetThumbnailBrush(const FAssetData& InAssetData, const FName InClassName) const
+{
+	return FPerformanceCaptureStyle::Get().GetBrush("ClassThumbnail.PCapCharacterDataAsset");
+}
+
+/*------------------------------------------------------------------------------
+	UPCapDataAsset Prop implementation.
+------------------------------------------------------------------------------*/
+const FSlateBrush* UAssetDefinition_PropDataAsset::GetIconBrush(const FAssetData& InAssetData, const FName InClassName) const
+{
+	return FPerformanceCaptureStyle::Get().GetBrush("ClassIcon.PCapPropDataAsset");
+}
+
+const FSlateBrush* UAssetDefinition_PropDataAsset::GetThumbnailBrush(const FAssetData& InAssetData, const FName InClassName) const
+{
+	return FPerformanceCaptureStyle::Get().GetBrush("ClassThumbnail.PCapPropDataAsset");
+}
+
+/*------------------------------------------------------------------------------
+	UPCapDataAsset Session Template implementation.
+------------------------------------------------------------------------------*/
+const FSlateBrush* UAssetDefinition_SessionTemplateAsset::GetIconBrush(const FAssetData& InAssetData, const FName InClassName) const
+{
+	return FPerformanceCaptureStyle::Get().GetBrush("ClassIcon.PCapSessionTemplate");
+}
+
+const FSlateBrush* UAssetDefinition_SessionTemplateAsset::GetThumbnailBrush(const FAssetData& InAssetData, const FName InClassName) const
+{
+	return FPerformanceCaptureStyle::Get().GetBrush("ClassThumbnail.PCapSessionTemplate");
+}
+
+#undef LOCTEXT_NAMESPACE
\ No newline at end of file
diff --git a/Engine/Plugins/VirtualProduction/PerformanceCaptureWorkflow/Source/PerformanceCaptureWorkflow/Private/PCapAssetDefinition.h b/Engine/Plugins/VirtualProduction/PerformanceCaptureWorkflow/Source/PerformanceCaptureWorkflow/Private/PCapAssetDefinition.h
new file mode 100644
--- /dev/null
+++ b/Engine/Plugins/VirtualProduction/PerformanceCaptureWorkflow/Source/PerformanceCaptureWorkflow/Private/PCapAssetDefinition.h
@@ -0,0 +1,99 @@
+// Copyright Epic Games, Inc. All Rights Reserved.
+
+#pragma once
+
+#include "AssetDefinition_DataAsset.h"
+#include "PCapDataTable.h"
+#include "PCapDatabase.h"
+#include "PCapSessionTemplate.h"
+#include "Table/AssetDefinition_DataTable.h"
+
+#include "PCapAssetDefinition.generated.h"
+
+/**
+ * 
+ */
+UCLASS()
+class PERFORMANCECAPTUREWORKFLOW_API UAssetDefinition_PCapDataTable : public UAssetDefinition_DataTable
+{
+	GENERATED_BODY()
+
+	// UAssetDefinition Begin
+	virtual FText GetAssetDisplayName() const override { return NSLOCTEXT("AssetDefinition", "PCapDataTable", "PCap Data Table"); }
+	virtual FLinearColor GetAssetColor() const override { return FLinearColor(FColor(57, 181, 74)); }
+	virtual TSoftClassPtr<UObject> GetAssetClass() const override { return UPCapDataTable::StaticClass(); }
+	virtual TConstArrayView<FAssetCategoryPath> GetAssetCategories() const override
+	{
+		static const auto Categories = { FAssetCategoryPath(EAssetCategoryPaths::Misc) };
+		return Categories;
+	}
+	virtual bool CanImport() const override { return true; }
+};
+
+UCLASS()
+class UAssetDefinition_PCapDataAsset : public UAssetDefinition_DataAsset
+{
+	GENERATED_BODY()
+public:
+	
+	// UAssetDefinition Implementation
+	virtual FText GetAssetDisplayName() const override { return NSLOCTEXT("AssetDefinition", "PCapDataAsset", "PCap DataAsset"); }
+	virtual FLinearColor GetAssetColor() const override { return FLinearColor(FColor(161, 57, 191)); }
+	virtual TSoftClassPtr<UObject> GetAssetClass() const override { return UPCapDataAsset::StaticClass(); }
+	// UAssetDefinition End
+};
+
+UCLASS()
+class UAssetDefinition_PerformerDataAsset : public UAssetDefinition_PCapDataAsset
+{
+	GENERATED_BODY()
+public:
+	
+	// UAssetDefinition Implementation
+	virtual FText GetAssetDisplayName() const override { return NSLOCTEXT("AssetDefinition", "PCapPerformerDataAsset", "PCap Performer Asset"); }
+	virtual const FSlateBrush* GetIconBrush(const FAssetData& InAssetData, const FName InClassName) const override;
+	virtual const FSlateBrush* GetThumbnailBrush(const FAssetData& InAssetData, const FName InClassName) const override;
+	virtual TSoftClassPtr<UObject> GetAssetClass() const override { return UPCapPerformerDataAsset::StaticClass(); }
+	// UAssetDefinition End
+};
+
+UCLASS()
+class UAssetDefinition_CharacterDataAsset : public UAssetDefinition_PCapDataAsset
+{
+	GENERATED_BODY()
+public:
+	
+	// UAssetDefinition Implementation
+	virtual FText GetAssetDisplayName() const override { return NSLOCTEXT("AssetDefinition", "PCapCharacterDataAsset", "PCap Character Asset"); }
+	virtual const FSlateBrush* GetIconBrush(const FAssetData& InAssetData, const FName InClassName) const override;
+	virtual const FSlateBrush* GetThumbnailBrush(const FAssetData& InAssetData, const FName InClassName) const override;
+	virtual TSoftClassPtr<UObject> GetAssetClass() const override { return UPCapCharacterDataAsset::StaticClass(); }
+	// UAssetDefinition End
+};
+
+UCLASS()
+class UAssetDefinition_PropDataAsset : public UAssetDefinition_PCapDataAsset
+{
+	GENERATED_BODY()
+public:
+	
+	// UAssetDefinition Implementation
+	virtual FText GetAssetDisplayName() const override { return NSLOCTEXT("AssetDefinition", "PCapPropDataAsset", "PCap Prop Asset"); }
+	virtual const FSlateBrush* GetIconBrush(const FAssetData& InAssetData, const FName InClassName) const override;
+	virtual const FSlateBrush* GetThumbnailBrush(const FAssetData& InAssetData, const FName InClassName) const override;
+	virtual TSoftClassPtr<UObject> GetAssetClass() const override { return UPCapPropDataAsset::StaticClass(); }
+	// UAssetDefinition End
+};
+UCLASS()
+class UAssetDefinition_SessionTemplateAsset : public UAssetDefinition_PCapDataAsset
+{
+	GENERATED_BODY()
+public:
+	
+	// UAssetDefinition Implementation
+	virtual FText GetAssetDisplayName() const override { return NSLOCTEXT("AssetDefinition", "PCapSessionTemplate", "PCap Session Template"); }
+	virtual const FSlateBrush* GetIconBrush(const FAssetData& InAssetData, const FName InClassName) const override;
+	virtual const FSlateBrush* GetThumbnailBrush(const FAssetData& InAssetData, const FName InClassName) const override;
+	virtual TSoftClassPtr<UObject> GetAssetClass() const override { return UPCapSessionTemplate::StaticClass(); }
+	// UAssetDefinition End
+};
\ No newline at end of file
diff --git a/Engine/Plugins/VirtualProduction/PerformanceCaptureWorkflow/Source/PerformanceCaptureWorkflow/Private/PCapAssetFactory.cpp b/Engine/Plugins/VirtualProduction/PerformanceCaptureWorkflow/Source/PerformanceCaptureWorkflow/Private/PCapAssetFactory.cpp
new file mode 100644
--- /dev/null
+++ b/Engine/Plugins/VirtualProduction/PerformanceCaptureWorkflow/Source/PerformanceCaptureWorkflow/Private/PCapAssetFactory.cpp
@@ -0,0 +1,330 @@
+// Copyright Epic Games, Inc. All Rights Reserved.
+
+
+#include "PCapAssetFactory.h"
+#include "PCapDatabase.h"
+#include "AssetToolsModule.h"
+#include "ClassViewerModule.h"
+#include "ClassViewerFilter.h"
+#include "IAssetTools.h"
+#include "PCapDataTable.h"
+#include "StructViewerFilter.h"
+#include "StructViewerModule.h"
+#include "Editor.h"
+#include "PerformanceCaptureStyle.h"
+
+#include "Modules/ModuleManager.h"
+
+#include "Widgets/SWindow.h"
+#include "Widgets/Input/SButton.h"
+#include "Widgets/Input/SComboButton.h"
+#include "Widgets/Text/STextBlock.h"
+#include "Widgets/Layout/SBox.h"
+#include "Styling/AppStyle.h"
+#include "Input/Reply.h"
+#include "Kismet2/SClassPickerDialog.h"
+
+#include UE_INLINE_GENERATED_CPP_BY_NAME(PCapAssetFactory)
+
+/*------------------------------------------------------------------------------
+	UPCap_DataAssetFactory implementation.
+------------------------------------------------------------------------------*/
+
+#define LOCTEXT_NAMESPACE "PerformanceCaptureDatatableFactory"
+
+UPCapDataTableFactory::UPCapDataTableFactory(const FObjectInitializer& ObjectInitializer)
+	: Super(ObjectInitializer)
+{
+	SupportedClass = UPCapDataTable::StaticClass();
+	bCreateNew = true;
+	bEditAfterNew = true;
+	Struct = FPCapRecordBase::StaticStruct();
+}
+
+bool UPCapDataTableFactory::ConfigureProperties()
+{
+	class FDataTableStructFilter : public IStructViewerFilter
+	{
+	public:
+		virtual bool IsStructAllowed(const FStructViewerInitializationOptions& InInitOptions, const UScriptStruct* InStruct, TSharedRef<FStructViewerFilterFuncs> InFilterFuncs) override
+		{
+			return (InStruct->IsChildOf(FPCapRecordBase::StaticStruct()) && (InStruct !=FPCapRecordBase::StaticStruct()));
+		}
+
+		virtual bool IsUnloadedStructAllowed(const FStructViewerInitializationOptions& InInitOptions, const FSoftObjectPath& InStructPath, TSharedRef<FStructViewerFilterFuncs> InFilterFuncs) override
+		{
+			// Unloaded structs are always User Defined Structs, and User Defined Structs are always allowed
+			// They will be re-validated by IsStructAllowed once loaded during the pick
+			return false;
+		}
+	};
+
+	class FDataTableFactoryUI : public TSharedFromThis<FDataTableFactoryUI>
+	{
+	public:
+		FReply OnCreate()
+		{
+			check(ResultStruct);
+			if (PickerWindow.IsValid())
+			{
+				PickerWindow->RequestDestroyWindow();
+			}
+			return FReply::Handled();
+		}
+
+		FReply OnCancel()
+		{
+			ResultStruct = nullptr;
+			if (PickerWindow.IsValid())
+			{
+				PickerWindow->RequestDestroyWindow();
+			}
+			return FReply::Handled();
+		}
+
+		bool IsStructSelected() const
+		{
+			return ResultStruct != nullptr;
+		}
+
+		void OnPickedStruct(const UScriptStruct* ChosenStruct)
+		{
+			ResultStruct = ChosenStruct;
+			StructPickerAnchor->SetIsOpen(false);
+		}
+
+		FText OnGetComboTextValue() const
+		{
+			return ResultStruct
+				? FText::AsCultureInvariant(ResultStruct->GetName())
+				: LOCTEXT("None", "None");
+		}
+
+		TSharedRef<SWidget> GenerateStructPicker()
+		{
+			FStructViewerModule& StructViewerModule = FModuleManager::LoadModuleChecked<FStructViewerModule>("StructViewer");
+
+			// Fill in options
+			FStructViewerInitializationOptions Options;
+			Options.Mode = EStructViewerMode::StructPicker;
+			Options.StructFilter = MakeShared<FDataTableStructFilter>();
+
+			return
+				SNew(SBox)
+				.WidthOverride(330.0f)
+				[
+					SNew(SVerticalBox)
+
+					+SVerticalBox::Slot()
+					.FillHeight(1.0f)
+					.MaxHeight(500)
+					[
+						SNew(SBorder)
+						.Padding(4)
+						.BorderImage(FAppStyle::GetBrush("ToolPanel.GroupBorder"))
+						[
+							StructViewerModule.CreateStructViewer(Options, FOnStructPicked::CreateSP(this, &FDataTableFactoryUI::OnPickedStruct))
+						]
+					]
+				];
+		}
+
+		const UScriptStruct* OpenStructSelector()
+		{
+			FStructViewerModule& StructViewerModule = FModuleManager::LoadModuleChecked<FStructViewerModule>("StructViewer");
+			ResultStruct = nullptr;
+
+			// Fill in options
+			FStructViewerInitializationOptions Options;
+			Options.Mode = EStructViewerMode::StructPicker;
+			Options.StructFilter = MakeShared<FDataTableStructFilter>();
+
+			PickerWindow = SNew(SWindow)
+				.Title(LOCTEXT("PCapDataTableFactoryOptions", "Pick Performance Capture Table Row Structure"))
+				.ClientSize(FVector2D(350, 100))
+				.SupportsMinimize(false)
+				.SupportsMaximize(false)
+				[
+					SNew(SBorder)
+					.BorderImage(FAppStyle::GetBrush("Menu.Background"))
+					.Padding(10)
+					[
+						SNew(SVerticalBox)
+						+SVerticalBox::Slot()
+						.AutoHeight()
+						[
+							SAssignNew(StructPickerAnchor, SComboButton)
+							.ContentPadding(FMargin(2,2,2,1))
+							.MenuPlacement(MenuPlacement_BelowAnchor)
+							.ButtonContent()
+							[
+								SNew(STextBlock)
+								.Text(this, &FDataTableFactoryUI::OnGetComboTextValue)
+							]
+							.OnGetMenuContent(this, &FDataTableFactoryUI::GenerateStructPicker)
+						]
+						+SVerticalBox::Slot()
+						.HAlign(HAlign_Right)
+						.AutoHeight()
+						[
+							SNew(SHorizontalBox)
+							+SHorizontalBox::Slot()
+							.AutoWidth()
+							[
+								SNew(SButton)
+								.Text(LOCTEXT("OK", "OK"))
+								.IsEnabled(this, &FDataTableFactoryUI::IsStructSelected)
+								.OnClicked(this, &FDataTableFactoryUI::OnCreate)
+							]
+							+SHorizontalBox::Slot()
+							.AutoWidth()
+							[
+								SNew(SButton)
+								.Text(LOCTEXT("Cancel", "Cancel"))
+								.OnClicked(this, &FDataTableFactoryUI::OnCancel)
+							]
+						]
+					]
+				];
+
+			GEditor->EditorAddModalWindow(PickerWindow.ToSharedRef());
+			PickerWindow.Reset();
+
+			return ResultStruct;
+		}
+
+	private:
+		TSharedPtr<SWindow> PickerWindow;
+		TSharedPtr<SComboButton> StructPickerAnchor;
+		const UScriptStruct* ResultStruct = nullptr;
+	};
+
+	TSharedRef<FDataTableFactoryUI> StructSelector = MakeShared<FDataTableFactoryUI>();
+	Struct = StructSelector->OpenStructSelector();
+
+	return Struct != nullptr;
+}
+
+FText UPCapDataTableFactory::GetDisplayName() const
+{
+	return LOCTEXT("PCapDataTable_DisplayName", "PCap Data Table");
+}
+
+FText UPCapDataTableFactory::GetToolTip() const
+{
+	return LOCTEXT("PerformanceCaptureDatatable_Tooltip", "Create a datatable for Performance Capture data management");
+}
+
+bool UPCapDataTableFactory::ShouldShowInNewMenu() const
+{
+	return true;
+}
+
+uint32 UPCapDataTableFactory::GetMenuCategories() const
+{
+	IAssetTools& AssetTools = FModuleManager::LoadModuleChecked<FAssetToolsModule>("AssetTools").Get();
+	return AssetTools.RegisterAdvancedAssetCategory("PerformanceCapture", LOCTEXT("AssetCategoryName", "Performance Capture"));
+}
+
+FName UPCapDataTableFactory::GetNewAssetThumbnailOverride() const
+{
+	return "ClassThumbnail.PCapDataTable";
+}
+
+UDataTable* UPCapDataTableFactory::MakeNewDataTable(UObject* InParent, FName Name, EObjectFlags Flags)
+{
+	return NewObject<UPCapDataTable>(InParent, Name, Flags);
+}
+
+#undef  LOCTEXT_NAMESPACE
+
+/*------------------------------------------------------------------------------
+	UPCapData Class Filter implementation.
+------------------------------------------------------------------------------*/
+
+class FPCapClassParentFilter : public IClassViewerFilter
+{
+public:
+	FPCapClassParentFilter()
+	: DisallowedClassFlags(CLASS_Abstract | CLASS_Deprecated | CLASS_NewerVersionExists | CLASS_HideDropDown) {}
+
+	/** All children of these classes will be included unless filtered out by another setting. */
+	TSet<const UClass*> AllowedChildrenOfClasses;
+
+	/** Disallowed class flags. */
+	EClassFlags DisallowedClassFlags;
+
+	virtual bool IsClassAllowed(const FClassViewerInitializationOptions& InInitOptions, const UClass* InClass, TSharedRef<FClassViewerFilterFuncs> InFilterFuncs) override
+	{
+		return !InClass->HasAnyClassFlags(DisallowedClassFlags)
+			&& InFilterFuncs->IfInChildOfClassesSet(AllowedChildrenOfClasses, InClass) != EFilterReturn::Failed;
+	}
+
+	virtual bool IsUnloadedClassAllowed(const FClassViewerInitializationOptions& InInitOptions, const TSharedRef<const IUnloadedBlueprintData> InUnloadedClassData, TSharedRef<FClassViewerFilterFuncs> InFilterFuncs) override
+	{
+		return !InUnloadedClassData->HasAnyClassFlags(DisallowedClassFlags)
+			&& InFilterFuncs->IfInChildOfClassesSet(AllowedChildrenOfClasses, InUnloadedClassData) != EFilterReturn::Failed;
+	}
+};
+
+/*------------------------------------------------------------------------------
+	UPCapDataAssetFactory implementation.
+------------------------------------------------------------------------------*/
+
+#define LOCTEXT_NAMESPACE "PerformanceCaptureDataAssetFactory"
+
+UPCap_DataAssetFactory::UPCap_DataAssetFactory(const FObjectInitializer& ObjectInitializer)
+	: Super(ObjectInitializer)
+{
+	bCreateNew = true;
+	bEditAfterNew = true;
+	SupportedClass = UPCapDataAsset::StaticClass();
+}
+
+bool UPCap_DataAssetFactory::ConfigureProperties()
+{
+	// nullptr the DataAssetClass so we can check for selection
+	DataAssetClass = nullptr;
+
+	// Load the classviewer module to display a class picker
+	FClassViewerModule& ClassViewerModule = FModuleManager::LoadModuleChecked<FClassViewerModule>("ClassViewer");
+
+	// Fill in options
+	FClassViewerInitializationOptions Options;
+	Options.Mode = EClassViewerMode::ClassPicker;
+
+	TSharedPtr<FPCapClassParentFilter> Filter = MakeShareable(new FPCapClassParentFilter);
+	Filter->AllowedChildrenOfClasses.Add(UPCapDataAsset::StaticClass());
+
+	Options.ClassFilters.Add(Filter.ToSharedRef());
+
+	const FText TitleText = LOCTEXT("CreatePCapDataAsset", "Pick Class For Performance Capture Data Asset");
+	UClass* ChosenClass = nullptr;
+	const bool bPressedOk = SClassPickerDialog::PickClass(TitleText, Options, ChosenClass, UPCapDataAsset::StaticClass());
+
+	if (bPressedOk)
+	{
+		DataAssetClass = ChosenClass;
+	}
+	return bPressedOk;
+}
+
+UObject* UPCap_DataAssetFactory::FactoryCreateNew(UClass* Class, UObject* InParent, FName Name, EObjectFlags Flags, UObject* Context, FFeedbackContext* Warn)
+{
+	UPCapDataAsset* NewDataAsset = nullptr;
+
+	if(DataAssetClass != nullptr)
+	{
+		NewDataAsset = NewObject<UPCapDataAsset>(InParent, DataAssetClass, Name, Flags | RF_Transactional, Context);
+	}
+	check(NewDataAsset);
+	return NewDataAsset;
+}
+
+uint32 UPCap_DataAssetFactory::GetMenuCategories() const
+{
+	IAssetTools& AssetTools = FModuleManager::LoadModuleChecked<FAssetToolsModule>("AssetTools").Get();
+	return AssetTools.RegisterAdvancedAssetCategory("PerformanceCapture", LOCTEXT("AssetCategoryName", "Performance Capture"));
+}
+
+#undef  LOCTEXT_NAMESPACE
\ No newline at end of file
diff --git a/Engine/Plugins/VirtualProduction/PerformanceCaptureWorkflow/Source/PerformanceCaptureWorkflow/Private/PCapAssetFactory.h b/Engine/Plugins/VirtualProduction/PerformanceCaptureWorkflow/Source/PerformanceCaptureWorkflow/Private/PCapAssetFactory.h
new file mode 100644
--- /dev/null
+++ b/Engine/Plugins/VirtualProduction/PerformanceCaptureWorkflow/Source/PerformanceCaptureWorkflow/Private/PCapAssetFactory.h
@@ -0,0 +1,63 @@
+// Copyright Epic Games, Inc. All Rights Reserved.
+
+#pragma once
+
+#include "CoreMinimal.h"
+#include "Factories/DataTableFactory.h"
+#include "UObject/ObjectMacros.h"
+#include "Templates/SubclassOf.h"
+#include "Factories/Factory.h"
+#include "PCapAssetFactory.generated.h"
+
+class UDataTable;
+class UPCapDataAsset;
+class UPCapCharacterDataAsset;
+class UPCapPropDataAsset;
+class UPCap_PropProxyDataAsset;
+
+
+/**
+ * Custom PCapDataTable Factory 
+ */
+UCLASS(BlueprintType)
+class UPCapDataTableFactory : public UDataTableFactory
+{
+	GENERATED_BODY()
+	
+public:
+
+	UPCapDataTableFactory(const FObjectInitializer& ObjectInitializer);
+
+	//~ Begin UFactory Interface
+	virtual bool ConfigureProperties() override;
+	virtual FText GetDisplayName() const override;
+	virtual FText GetToolTip() const override;
+	virtual bool ShouldShowInNewMenu() const override;
+	virtual uint32 GetMenuCategories() const override;
+	virtual FName GetNewAssetThumbnailOverride() const override;
+	//~ Begin UFactory Interface
+
+protected:
+	
+	virtual UDataTable* MakeNewDataTable(UObject* InParent, FName Name, EObjectFlags Flags) override;
+};
+
+UCLASS(hidecategories=Object, MinimalAPI)
+
+class UPCap_DataAssetFactory : public UFactory
+{
+	GENERATED_BODY()
+	
+public:
+	UPCap_DataAssetFactory(const FObjectInitializer& ObjectInitializer);
+	
+	UPROPERTY(EditAnywhere, Category=DataAsset)
+	TSubclassOf<UPCapDataAsset> DataAssetClass;
+
+	// UFactory interface
+	virtual bool ConfigureProperties() override;
+	virtual UObject* FactoryCreateNew(UClass* Class, UObject* InParent, FName Name, EObjectFlags Flags, UObject* Context, FFeedbackContext* Warn) override;
+	virtual uint32 GetMenuCategories() const override;
+	// End of UFactory interface
+	
+};
\ No newline at end of file
diff --git a/Engine/Plugins/VirtualProduction/PerformanceCaptureWorkflow/Source/PerformanceCaptureWorkflow/Private/PCapBPFunctionLibrary.cpp b/Engine/Plugins/VirtualProduction/PerformanceCaptureWorkflow/Source/PerformanceCaptureWorkflow/Private/PCapBPFunctionLibrary.cpp
new file mode 100644
--- /dev/null
+++ b/Engine/Plugins/VirtualProduction/PerformanceCaptureWorkflow/Source/PerformanceCaptureWorkflow/Private/PCapBPFunctionLibrary.cpp
@@ -0,0 +1,92 @@
+// Copyright Epic Games, Inc. All Rights Reserved.
+
+#include "PCapBPFunctionLibrary.h"
+#include "TakesCore/Public/TakeRecorderSource.h"
+#include "Engine/Engine.h"
+#include "EngineUtils.h"
+
+
+FString UPerformanceCaptureBPFunctionLibrary::SanitizeFileString(FString InString)
+{
+	FString BannedChars = "[]{}!\"$£%^&*()+=;:?\\/|'@#~ ";
+	FString CleanString = InString;
+
+	//BannedChars = BannedChars + FPaths::GetInvalidFileSystemChars(); //This probably causes duplicates but rather than consider each platform here, we'll rely on the core.
+	TArray<FString> IgnoredChars;
+
+	for(auto CharIt(BannedChars.CreateConstIterator()); CharIt; ++CharIt)
+	{
+		TCHAR Char = *CharIt;
+		IgnoredChars.Add(FString(1, &Char));
+	}
+	// Remove the null terminator on the end
+	IgnoredChars.RemoveAt(IgnoredChars.Num()-1, 1);
+	
+	for(FString IgnoredChar : IgnoredChars)
+	{
+		CleanString = CleanString.Replace(*IgnoredChar, TEXT(""), ESearchCase::IgnoreCase);
+	}
+	return CleanString;
+}
+
+FString UPerformanceCaptureBPFunctionLibrary::SanitizePathString(FString InString)
+{
+	{
+		//This is the same as for file name strings but with allowed exception of "/" directory limiters
+		//TODO: currently windows only!
+		FString BannedChars = "[]{}!\"$£%^&*()+=;:?|'@#~ ";
+		FString CleanString = InString;
+
+		//BannedChars = BannedChars + FPaths::GetInvalidFileSystemChars(); //This probably causes duplicates but rather than consider each platform here, we'll rely on the core.
+		TArray<FString> IgnoredChars;
+
+		for(auto CharIt(BannedChars.CreateConstIterator()); CharIt; ++CharIt)
+		{
+			TCHAR Char = *CharIt;
+			IgnoredChars.Add(FString(1, &Char));
+		}
+		// Remove the null terminator on the end
+		IgnoredChars.RemoveAt(IgnoredChars.Num()-1, 1);
+	
+		for(FString IgnoredChar : IgnoredChars)
+		{
+			CleanString = CleanString.Replace(*IgnoredChar, TEXT(""), ESearchCase::IgnoreCase);
+		}
+		return CleanString;
+	}
+}
+
+TArray<UObject*> UPerformanceCaptureBPFunctionLibrary::GetTakeRecorderSourceSettingsObjects(UTakeRecorderSource* InSource, bool& bSuccess)
+{
+	if(InSource)
+	{
+		bSuccess = true;
+		return InSource->GetAdditionalSettingsObjects();
+	}
+	TArray<UObject*> NullResult;
+	bSuccess = false;
+	return NullResult;
+}
+
+void UPerformanceCaptureBPFunctionLibrary::GetAllActorsWithComponent(const UObject* WorldContextObject, TSubclassOf<UActorComponent> Component, TArray<AActor*>& OutActors)
+{
+	QUICK_SCOPE_CYCLE_COUNTER(ITWBPFactory_GetAllActorsWithComponent);
+	OutActors.Reset();
+    
+	if (!Component)
+	{
+		return;
+	}
+
+	if (UWorld* World = GEngine->GetWorldFromContextObject(WorldContextObject, EGetWorldErrorMode::LogAndReturnNull))
+	{
+		for (FActorIterator It(World); It; ++It)
+		{
+			AActor* Actor = *It;
+			if (Actor->GetComponentByClass(Component) != nullptr)
+			{
+				OutActors.Add(Actor);
+			}
+		}
+	}
+}
diff --git a/Engine/Plugins/VirtualProduction/PerformanceCaptureWorkflow/Source/PerformanceCaptureWorkflow/Private/PCapBPFunctionLibrary.h b/Engine/Plugins/VirtualProduction/PerformanceCaptureWorkflow/Source/PerformanceCaptureWorkflow/Private/PCapBPFunctionLibrary.h
new file mode 100644
--- /dev/null
+++ b/Engine/Plugins/VirtualProduction/PerformanceCaptureWorkflow/Source/PerformanceCaptureWorkflow/Private/PCapBPFunctionLibrary.h
@@ -0,0 +1,55 @@
+// Copyright Epic Games, Inc. All Rights Reserved.
+
+#pragma once
+
+#include "CoreMinimal.h"
+#include "Editor/Blutility/Private/EditorFunctionLibrary.h"
+#include "Components/ActorComponent.h"
+#include "PCapBPFunctionLibrary.generated.h"
+
+class UTakeRecorderSource;
+/**
+ * 
+ */
+UCLASS()
+class PERFORMANCECAPTUREWORKFLOW_API UPerformanceCaptureBPFunctionLibrary : public UEditorFunctionLibrary
+{
+	GENERATED_BODY()
+	
+public:
+	/**
+	 * Returns a string stripped of the following characters []{}1!"$£%^&*()+=;:?/\|'@#~
+	 * 
+	 * @param InString String you wish to clean
+	 * @return Cleaned string
+	 */
+	UFUNCTION(BlueprintCallable, Category = "Performance Capture|Workflow")
+	static FString SanitizeFileString(FString InString);
+	
+	/**
+	 * Returns a string stripped of the following characters []{}1!"$£%^&*()+=;:?\|'@#~ Note filepath delimiters are not excluded
+	 * 
+	 * @param InString String you wish to clean
+	 * @return Cleaned string 
+	 */
+	UFUNCTION(BlueprintCallable, Category = "Performance Capture|Workflow")
+	static FString SanitizePathString(FString InString);
+
+	/**
+	 * 
+	 * @param InSource TakeRecorderSourceObject
+	 * @param bSuccess Whether the Settings objects were successfully retrieved 
+	 * @return 
+	 */
+	UFUNCTION(BlueprintCallable, Category = "Performance Capture|Workflow")
+	static TArray<UObject*> GetTakeRecorderSourceSettingsObjects(UTakeRecorderSource* InSource,bool& bSuccess);
+	
+	/** 
+	*	Find all Actors in the world containing at least one instance of the given component class. 
+	*	This is a very slow operation, as it will search over every actor in the world.
+	*	@param	Component	Class of component to find. Must be specified or result array will be empty.
+	*	@param	OutActors	Output array of Actors of the specified class.
+	 */
+	UFUNCTION(BlueprintCallable, Category="Performance Capture|Core",  meta=(WorldContext="WorldContextObject", DeterminesOutputType="Component", DynamicOutputParam="OutActors"))
+	static void GetAllActorsWithComponent(const UObject* WorldContextObject, TSubclassOf<UActorComponent> Component, TArray<AActor*>& OutActors);
+};
diff --git a/Engine/Plugins/VirtualProduction/PerformanceCaptureWorkflow/Source/PerformanceCaptureWorkflow/Private/PCapDataTable.cpp b/Engine/Plugins/VirtualProduction/PerformanceCaptureWorkflow/Source/PerformanceCaptureWorkflow/Private/PCapDataTable.cpp
new file mode 100644
--- /dev/null
+++ b/Engine/Plugins/VirtualProduction/PerformanceCaptureWorkflow/Source/PerformanceCaptureWorkflow/Private/PCapDataTable.cpp
@@ -0,0 +1,24 @@
+// Copyright Epic Games, Inc. All Rights Reserved.
+
+
+#include "PCapDataTable.h"
+
+#include "PCapDatabase.h"
+
+UPCapDataTable::UPCapDataTable()
+{
+	RowStruct = FPCapRecordBase::StaticStruct();
+	bStripFromClientBuilds = true; //Prevent any Performance Capture data from cooking into client builds as this is editor-only tooling
+
+	OnDataTableChanged().AddUObject(this, &UPCapDataTable::OnDataTableModified);
+}
+
+UPCapDataTable::~UPCapDataTable()
+{
+	
+}
+
+void UPCapDataTable::OnDataTableModified() const
+{
+	OnDatatableModified.Broadcast();
+}
diff --git a/Engine/Plugins/VirtualProduction/PerformanceCaptureWorkflow/Source/PerformanceCaptureWorkflow/Private/PCapDataTable.h b/Engine/Plugins/VirtualProduction/PerformanceCaptureWorkflow/Source/PerformanceCaptureWorkflow/Private/PCapDataTable.h
new file mode 100644
--- /dev/null
+++ b/Engine/Plugins/VirtualProduction/PerformanceCaptureWorkflow/Source/PerformanceCaptureWorkflow/Private/PCapDataTable.h
@@ -0,0 +1,27 @@
+// Copyright Epic Games, Inc. All Rights Reserved.
+
+#pragma once
+
+#include "CoreMinimal.h"
+#include "Engine/DataTable.h"
+#include "PCapDataTable.generated.h"
+
+/**
+ * 
+ */
+UCLASS()
+class PERFORMANCECAPTUREWORKFLOW_API UPCapDataTable : public UDataTable
+{
+	GENERATED_BODY()
+	
+public:
+	UPCapDataTable();
+	virtual ~UPCapDataTable() override;
+
+	DECLARE_DYNAMIC_MULTICAST_DELEGATE(FOnDatatableModified);
+	UPROPERTY(BlueprintAssignable, Category="Performance Capture|Database")
+	FOnDatatableModified OnDatatableModified;
+
+	/** Emit a callback each time the datatable is modified*/
+	void OnDataTableModified() const;
+};
diff --git a/Engine/Plugins/VirtualProduction/PerformanceCaptureWorkflow/Source/PerformanceCaptureWorkflow/Private/PCapDatabase.cpp b/Engine/Plugins/VirtualProduction/PerformanceCaptureWorkflow/Source/PerformanceCaptureWorkflow/Private/PCapDatabase.cpp
new file mode 100644
--- /dev/null
+++ b/Engine/Plugins/VirtualProduction/PerformanceCaptureWorkflow/Source/PerformanceCaptureWorkflow/Private/PCapDatabase.cpp
@@ -0,0 +1,154 @@
+// Copyright Epic Games, Inc. All Rights Reserved.
+
+#include "PCapDatabase.h"
+#include "CapturePerformer.h"
+#include "CaptureCharacter.h"
+#include "PerformanceCapture.h"
+#include "Misc/Guid.h"
+
+
+FPCapRecordBase::FPCapRecordBase()
+	: UID(FGuid::NewGuid())
+{
+}
+
+FPCapRecordBase::~FPCapRecordBase()
+{
+}
+
+void FPCapRecordBase::OnDataTableChanged(const UDataTable* InDataTable, const FName InRowName)
+{
+	FTableRowBase::OnDataTableChanged(InDataTable, InRowName);
+}
+
+void FPCapRecordBase::OnPostDataImport(const UDataTable* InDataTable, const FName InRowName, TArray<FString>& OutCollectedImportProblems)
+{
+	//TODO Handle not over-writing UID if a user imports
+	FTableRowBase::OnPostDataImport(InDataTable, InRowName, OutCollectedImportProblems);
+}
+
+UPCapDataAsset::UPCapDataAsset()
+{
+	CreateGuid();
+}
+
+void UPCapDataAsset::CreateGuid()
+{
+	if (!AssetUID.IsValid())
+	{
+		AssetUID = FGuid::NewGuid();
+	}
+}
+
+void UPCapDataAsset::PostDuplicate(EDuplicateMode::Type DuplicateMode)
+{
+	Super::PostDuplicate(DuplicateMode);
+
+	AssetUID = FGuid::NewGuid();
+}
+
+/** Performer DataAsset constructor */
+UPCapPerformerDataAsset::UPCapPerformerDataAsset()
+	: UPCapDataAsset()
+	, PerformerActorClass(ACapturePerformer::StaticClass())
+	, PerformerColor(1.0f, 1.0f, 1.0f, 1.0f)
+
+{
+}
+
+/** Character DataAsset constructor */
+UPCapCharacterDataAsset::UPCapCharacterDataAsset()
+	: UPCapDataAsset()
+	, CaptureCharacterClass(ACaptureCharacter::StaticClass())
+{
+}
+
+UPCapPropDataAsset::UPCapPropDataAsset()
+	: UPCapDataAsset()
+{
+}
+
+/** Prop DataAsset PostEditChange */
+void UPCapPropDataAsset::PostEditChangeProperty(FPropertyChangedEvent& PropertyChangedEvent)
+{
+	Super::PostEditChangeProperty(PropertyChangedEvent);
+
+	if(PropertyChangedEvent.Property->GetFName() == GET_MEMBER_NAME_CHECKED(UPCapPropDataAsset, bClearEditConditions))
+	{
+		bClearEditConditions = false;
+		bCanSetStaticMesh = true;
+		bCanSetSkeletalMesh = true;
+		bCanSetCustomPropClass =true;
+	}
+	if(PropertyChangedEvent.Property->GetFName() == GET_MEMBER_NAME_CHECKED(UPCapPropDataAsset, PropStaticMesh))
+	{
+		FString PrintMe = PropStaticMesh.ToString();
+		bool ValidThing = PropStaticMesh.IsNull();
+		UE_LOG(LogPCap, Display, TEXT("Static Mesh is %s"), ( ValidThing ? TEXT("true") : TEXT("false") ));
+		
+		if(!PropStaticMesh.IsNull())
+		{
+			bCanSetSkeletalMesh = false;
+			bCanSetCustomPropClass = false;
+		}
+		else if(PropStaticMesh.IsNull())
+		{
+			bCanSetStaticMesh = true;
+			bCanSetSkeletalMesh = true;
+			bCanSetCustomPropClass =true;
+		}
+	}
+	if(PropertyChangedEvent.Property->GetFName() == GET_MEMBER_NAME_CHECKED(UPCapPropDataAsset, PropSkeletalMesh))
+	{
+		if(!PropSkeletalMesh.IsNull())
+		{
+			bCanSetStaticMesh = false;
+			bCanSetCustomPropClass = false;
+		}
+		else if(PropSkeletalMesh.IsNull())
+		{
+			bCanSetStaticMesh = true;
+			bCanSetSkeletalMesh = true;
+			bCanSetCustomPropClass =true;
+		}
+	}
+	if(PropertyChangedEvent.Property->GetFName() == GET_MEMBER_NAME_CHECKED(UPCapPropDataAsset, CustomPropClass))
+	{
+		if(!CustomPropClass.IsNull())
+		{
+			bCanSetStaticMesh = false;
+			bCanSetSkeletalMesh = false;
+		}
+		else if(CustomPropClass.IsNull())
+		{
+			bCanSetStaticMesh = true;
+			bCanSetSkeletalMesh = true;
+			bCanSetCustomPropClass =true;
+		}
+	}
+}
+
+void UPCapPropDataAsset::PostLoad()
+{
+	Super::PostLoad();
+	ValidateEditConditions();
+}
+
+void UPCapPropDataAsset::ValidateEditConditions()
+{
+	if(!CustomPropClass.IsNull())
+    {
+    	bCanSetStaticMesh = false;
+    	bCanSetSkeletalMesh = false;
+    }
+	if(!PropStaticMesh.IsNull())
+	{
+		bCanSetSkeletalMesh = false;
+		bCanSetCustomPropClass = false;
+	}
+	if(!PropSkeletalMesh.IsNull())
+	{
+		bCanSetStaticMesh = false;
+		bCanSetCustomPropClass = false;
+	}	
+}
diff --git a/Engine/Plugins/VirtualProduction/PerformanceCaptureWorkflow/Source/PerformanceCaptureWorkflow/Private/PCapDatabase.h b/Engine/Plugins/VirtualProduction/PerformanceCaptureWorkflow/Source/PerformanceCaptureWorkflow/Private/PCapDatabase.h
new file mode 100644
--- /dev/null
+++ b/Engine/Plugins/VirtualProduction/PerformanceCaptureWorkflow/Source/PerformanceCaptureWorkflow/Private/PCapDatabase.h
@@ -0,0 +1,435 @@
+// Copyright Epic Games, Inc. All Rights Reserved.
+
+#pragma once
+
+#include "CoreMinimal.h"
+#include "CapturePerformer.h"
+#include "CaptureCharacter.h"
+#include "LevelSequence.h"
+#include "LiveLinkTypes.h"
+#include "PCapDataTable.h"
+#include "Engine/DataTable.h"
+#include "Engine/DataAsset.h"
+#include "Misc/DateTime.h"
+#include "UObject/ConstructorHelpers.h"
+#include "PCapDatabase.generated.h"
+
+class UPCapPropDataAsset;
+class UPCapCharacterDataAsset;
+class UPCapPerformerDataAsset;
+class UPCapSessionTemplate;
+class UPCapDataTable;
+class ULevelSequence;
+class UIKRigDefinition;
+class UIKRetargeter;
+class ULiveLinkInstance;
+class UDataLayerAsset;
+class USkeletalMesh;
+
+/// Database structs ///
+
+USTRUCT()
+struct FPCapRecordBase : public FTableRowBase
+{
+	GENERATED_BODY()
+
+	FPCapRecordBase();
+	virtual ~FPCapRecordBase() override;
+
+	//Only called when you import or re-import data. TODO Make sure the UIDs are not overwritten
+
+	virtual void OnDataTableChanged(const UDataTable* InDataTable, const FName InRowName) override;
+
+	virtual void OnPostDataImport(const UDataTable* InDataTable, const FName InRowName, TArray<FString>& OutCollectedImportProblems) override;
+	
+	UPROPERTY(VisibleAnywhere, BlueprintReadWrite, Category="Performance Capture", meta=(IgnoreForMemberInitializationTest))
+	FGuid UID;
+
+	UPROPERTY(EditAnywhere, BlueprintReadWrite, Category="Performance Capture")
+	bool bIsArchived = false;
+	
+};
+
+USTRUCT(BlueprintType)
+struct FPCapProductionRecord : public FPCapRecordBase
+{
+	GENERATED_BODY()
+
+public:
+	UPROPERTY(EditAnywhere, BlueprintReadWrite, Category="Performance Capture Production")
+	FName ProductionName;
+
+	UPROPERTY(EditAnywhere, BlueprintReadWrite, Category="Performance Capture Production")
+	FString ProductionNotes;
+};
+
+UENUM(BlueprintType)
+enum class EPCapTakeStatus : uint8
+{
+	Neutral,
+	ThumbsUp,
+	ThumbsDown
+};
+
+USTRUCT(BlueprintType)
+struct FPCapTakeRecord : public FPCapRecordBase
+{
+	GENERATED_BODY()
+	
+public:
+
+	/** Temporary Flag. Transient, so will not be saved */
+	UPROPERTY(Transient, BlueprintReadWrite, Category="Flag")
+	bool bFlag = false;
+	
+	/** The level sequence recorded */
+	UPROPERTY(EditAnywhere, BlueprintReadWrite, Category="Performance Capture Take")
+	TObjectPtr<ULevelSequence> RecordedTake;
+	
+	/** Seconds Duration */
+	UPROPERTY(VisibleAnywhere, BlueprintReadWrite, Category="Performance Capture Take")
+	float TakeDurationSeconds = 0;
+
+	/** Recorded Framerate*/
+	UPROPERTY(VisibleAnywhere, BlueprintReadWrite, Category="Performance Capture Take")
+	FFrameRate Framerate = FFrameRate(30, 1);
+	
+	/** HHMMSSFF Duration*/
+	UPROPERTY(VisibleAnywhere, BlueprintReadWrite, Category="Performance Capture Take")
+	FTimecode TakeDurationTimecode;
+
+	/** Start Timecode*/
+	UPROPERTY(VisibleAnywhere, BlueprintReadWrite, Category="Performance Capture Take")
+	FTimecode StartTimecode;
+
+	/** End Timecode*/
+	UPROPERTY(VisibleAnywhere, BlueprintReadWrite, Category="Performance Capture Take")
+	FTimecode EndTimecode;
+
+	/** Transform of the stage root when this recording was made */
+	UPROPERTY(EditAnywhere, BlueprintReadWrite, Category="Performance Capture Take")
+	FTransform MocapStageRootTransform;
+
+	/** Status of Take - Thumbs Up, Thumbs Down or None */
+	UPROPERTY(EditAnywhere, BlueprintReadWrite, Category="Performance Capture Take")
+	EPCapTakeStatus TakeStatus = EPCapTakeStatus::Neutral;
+
+	/** 5-Star Rating */
+	UPROPERTY(EditAnywhere, BlueprintReadWrite, Category="Performance Capture Take", meta=(ClampMin=0, ClampMax=5))
+	int32 Rating = 0;
+	
+	/** Does this recording LiveLinkSource tracks */
+	UPROPERTY(EditAnywhere, BlueprintReadWrite, Category="Performance Capture Take")
+	bool bContainsLiveLinkSources = false;
+
+	/** Has this take been processed for plotting */
+	UPROPERTY(EditAnywhere, BlueprintReadWrite, Category="Performance Capture Take")
+	bool bLiveLinkPlotted = false;
+	
+	/** Has the animation recorded for this take been replaced with an external asset import.
+	 * This option will be positive if you use the import and upgrade workflow
+	 * bringing fbx recordings from your motion capture software in and
+	 * replacing the live recording.
+	 */
+	UPROPERTY(EditAnywhere, BlueprintReadWrite, Category="Performance Capture Take")
+	bool bExternallyReplaced = false;
+
+	UPROPERTY(VisibleAnywhere, BlueprintReadWrite, Category="Performance Capture Take")
+	FGuid SessionUID;
+};
+
+USTRUCT(BlueprintType)
+struct FPCapSessionRecord : public FPCapRecordBase
+{
+	GENERATED_BODY()
+	
+public:
+	
+	/**Name for this session*/
+	UPROPERTY(VisibleAnywhere, BlueprintReadWrite, Category="Performance Capture Session")
+	FName SessionName;
+
+	UPROPERTY(VisibleAnywhere, BlueprintReadWrite, Category="Performance Capture Session")
+	FDateTime SessionDateTime;
+
+	UPROPERTY(EditAnywhere, BlueprintReadWrite, Category="Performance Capture Session")
+	FString SessionNotes;
+
+	/**ID of Production this Session belongs to */
+	UPROPERTY(VisibleAnywhere, BlueprintReadWrite, Category="Performance Capture Session")
+	FName ProductionName;
+
+	UPROPERTY(VisibleAnywhere, BlueprintReadWrite, Category="Performance Capture Session")
+	FGuid ProductionUID;
+
+	UPROPERTY(VisibleAnywhere, BlueprintReadWrite, Category="Performance Capture Session", meta = (DisplayName = "Session OutputName"))
+	FString SessionOutputName;
+
+	UPROPERTY(VisibleAnywhere, BlueprintReadWrite, Category="Performance Capture Session", meta = (DisplayName = "Session Folder", ContentDir))
+	FString SessionPath;
+
+	UPROPERTY(VisibleAnywhere, BlueprintReadWrite , Category="Performance Capture Session", meta = (DisplayName = "Performer Folder", ContentDir))
+	FString PerformerPath;
+
+	UPROPERTY(VisibleAnywhere, BlueprintReadWrite, Category="Performance Capture Session", meta = (DisplayName = "Character Folder", ContentDir))
+	FString CharacterPath;
+
+	UPROPERTY(VisibleAnywhere, BlueprintReadWrite, Category="Performance Capture Session", meta = (DisplayName = "Takes Folder", ContentDir))
+	FString TakesPath;
+	
+	UPROPERTY(VisibleAnywhere, BlueprintReadWrite, Category="Performance Capture Session", meta = (DisplayName = "Prop Folder", ContentDir))
+	FString PropPath;
+
+	UPROPERTY(VisibleAnywhere, BlueprintReadWrite, Category="Performance Capture Session", meta = (DisplayName = "Scene Folder", ContentDir))
+	FString ScenePath;
+
+	UPROPERTY(VisibleAnywhere, BlueprintReadWrite, Category="Performance Capture Session", meta = (DisplayName = "Common Folder", ContentDir))
+	FString CommonPath;
+
+	UPROPERTY(VisibleAnywhere, BlueprintReadWrite, Category="Performance Capture Session", meta = (DisplayName = "Additional Folders", ContentDir))
+	TArray<FString> AdditionalFolders;
+
+	UPROPERTY(VisibleAnywhere, BlueprintReadWrite, Category="Performance Capture Session")
+	TSoftObjectPtr<UDataTable> TakesDataTable;
+
+	UPROPERTY(VisibleAnywhere, BlueprintReadWrite, Category="Performance Capture Session")
+	TArray<TSoftObjectPtr<UPCapPerformerDataAsset>> Performers;
+
+	UPROPERTY(VisibleAnywhere, BlueprintReadWrite, Category="Performance Capture Session")
+	TArray<TSoftObjectPtr<UPCapCharacterDataAsset>> Characters;
+
+	UPROPERTY(VisibleAnywhere, BlueprintReadWrite, Category="Performance Capture Session")
+	TArray<TSoftObjectPtr<UPCapPropDataAsset>> Props;
+
+	UPROPERTY(VisibleAnywhere, BlueprintReadWrite, Category="Performance Capture Session")
+	TSoftObjectPtr<UPCapSessionTemplate> SessionTemplate;
+
+	UPROPERTY(VisibleAnywhere, BlueprintReadWrite, Category="Performance Capture Session")
+	bool bIsWorldPartition = false;
+
+	UPROPERTY(VisibleAnywhere, BlueprintReadWrite, Category="Performance Capture Session")
+	TSoftObjectPtr<UWorld> SubLevel;
+
+	UPROPERTY(VisibleAnywhere, BlueprintReadWrite, Category="Performance Capture Session")
+	TSoftObjectPtr<UDataLayerAsset> SessionDataLayer;
+
+	UPROPERTY(EditAnywhere, BlueprintReadWrite, Category="Performance Capture Session", meta = (RequiredAssetDataTags = "RowStructure=/Script/PerformanceCaptureWorkflow.PCapSlateRecord"))
+	TSoftObjectPtr<UPCapDataTable> SessionSlateTable;
+};
+
+UENUM(BlueprintType)
+enum class EPCapSlateStatus : uint8
+{
+	Incomplete	UMETA(Tooltip="Slate is incomplete"),
+	Complete	UMETA(Tooltip="Slate is marked complete"),
+	Skip		UMETA(Tooltip="Slate is marked to skip")
+};
+
+USTRUCT(BlueprintType)
+struct FPCapSlateRecord : public FPCapRecordBase
+{                                                                                                                                                                                                                                                                                                                                                                                                                                 
+	GENERATED_BODY()
+
+	//*Name of slate*/
+	UPROPERTY(EditAnywhere, BlueprintReadWrite, Category="Performance Capture Slate", meta=(DisplayName="Slate"))
+	FString Slate;
+	
+	//*Note on slate. This will be passed to Take Recorder*/
+	UPROPERTY(EditAnywhere, BlueprintReadWrite, Category="Performance Capture Slate", meta=(DisplayName="Note", MultiLine=true))
+	FString SlateNote;
+
+	//*Status of Slate*/
+	UPROPERTY(EditAnywhere, BlueprintReadWrite, Category="Performance Capture Slate", meta=(DisplayName="Status"))
+	EPCapSlateStatus SlateStatus = EPCapSlateStatus::Incomplete;
+
+	//*UID of the session this Slate is used in*/
+	UPROPERTY(VisibleAnywhere, BlueprintReadWrite, Category="Performance Capture Slate")
+	FGuid SessionUID;
+};
+
+///// End Database Structs //////////
+
+/////////////////////////////////////
+/// Begin Data Asset Definitions/////
+/////////////////////////////////////
+
+/// Data Assets are for saving collections of assets that should/need to be gathered together
+
+/** Performance Capture DataAsset. Contains references to assets used for Performance Capture Workflows*/
+UCLASS(Abstract, MinimalAPI)
+class UPCapDataAsset : public UPrimaryDataAsset
+{
+	GENERATED_BODY()
+	
+public:
+	/** Constructor*/
+	UPCapDataAsset();
+
+	void CreateGuid();
+
+	virtual void PostDuplicate(EDuplicateMode::Type DuplicateMode) override;
+
+	UFUNCTION(BlueprintGetter, Category="Performance Capture|Admin")
+	FGuid GetAssetUID() const {return AssetUID;}
+
+private:
+
+	/** Guid for disambiguating actors spawned by the data assets*/
+	UPROPERTY(BlueprintGetter="GetAssetUID", Category="Performance Capture|Admin", meta=(IgnoreForMemberInitializationTest))
+	FGuid AssetUID;
+};
+
+/** Mocap Performer DataAsset */
+UCLASS(Blueprintable, BlueprintType)
+class UPCapPerformerDataAsset : public UPCapDataAsset
+{
+	GENERATED_BODY()
+
+public:
+	/** Constructor*/
+	UPCapPerformerDataAsset();
+
+	/** LiveLink Subject for the body of this Performer */
+	UPROPERTY(BlueprintReadWrite, EditDefaultsOnly, Category="Performance Capture Performer")
+	FName PerformerName;
+	
+	/** LiveLink Subject for this Performer */
+	UPROPERTY(BlueprintReadWrite, EditDefaultsOnly, Category="Performance Capture Performer")
+	FLiveLinkSubjectName LiveLinkSubject;
+	
+	/** Performer Actor class to use for this performer */
+	UPROPERTY(BlueprintReadWrite, EditDefaultsOnly, NoClear, Category="Performance Capture Performer")
+	TSoftClassPtr<ACapturePerformer> PerformerActorClass;
+	
+	/** Color for this performer */
+	UPROPERTY(BlueprintReadWrite, EditDefaultsOnly, Category="Performance Capture Performer")
+	FLinearColor PerformerColor;
+	
+	/** Performer Base Skeleton */
+	UPROPERTY(BlueprintReadWrite, EditAnywhere, Category="Performance Capture Performer")
+	TSoftObjectPtr<USkeletalMesh> BaseSkeletalMesh;
+	
+	/** Performer Mesh. This must be created in a T or A Pose performer from a LiveLink pose*/
+	UPROPERTY(BlueprintReadWrite, EditAnywhere, Category="Performance Capture Performer")
+	TSoftObjectPtr<USkeletalMesh> PerformerProportionedMesh;
+
+	/** Performer IKRig */
+	UPROPERTY(BlueprintReadWrite, EditDefaultsOnly, Category="Performance Capture Performer")
+	TSoftObjectPtr<UIKRigDefinition> IKRig;
+};
+
+/** Mocap Character Data asset */
+UCLASS(Blueprintable, BlueprintType)
+class UPCapCharacterDataAsset : public UPCapDataAsset
+{
+	GENERATED_BODY()
+	
+public:
+
+	/** Capture Character Constructor */
+	UPCapCharacterDataAsset();
+
+	/** Capture Character Name */
+	UPROPERTY(BlueprintReadWrite, EditAnywhere, Category="Performance Capture Character", meta=(Tooltip="Optional Name for Character. If blank DataAsset name will be used"))
+	FName CharacterName;
+
+	/** Source Performer Asset for this Character */
+	UPROPERTY(BlueprintReadWrite, EditAnywhere, Category="Performance Capture Character")
+	TSoftObjectPtr<UPCapPerformerDataAsset> SourcePerformerAsset;
+
+	/** Character actor Class */
+	UPROPERTY(BlueprintReadWrite, EditAnywhere, Category="Performance Capture Character")
+	TSoftClassPtr<ACaptureCharacter>  CaptureCharacterClass;
+
+	/** Main Skeletal Mesh Asset. This will be the root component of any characters spawned. A character will not be spawned if this asset reference is null or not valid */
+	UPROPERTY(BlueprintReadWrite, EditAnywhere, Category="Performance Capture Character")
+	TSoftObjectPtr<USkeletalMesh> SkeletalMesh;
+	
+	/** IKRig Asset */
+	UPROPERTY(BlueprintReadWrite, EditAnywhere, Category="Performance Capture Character")
+	TSoftObjectPtr<UIKRigDefinition> IKRig;
+
+	/** Retarget Asset */
+	UPROPERTY(BlueprintReadWrite, EditAnywhere, Category="Performance Capture Character")
+	TSoftObjectPtr<UIKRetargeter> Retargeter;
+
+	UPROPERTY(BlueprintReadWrite, EditAnywhere, Category="Performance Capture Character", meta=(DisplayName="Additional Meshes", ToolTip="Additional meshes that will be added as child components"))
+	TArray<TSoftObjectPtr<USkeletalMesh>> AdditionalMeshes;
+};
+
+/**  Prop DataAsset*/
+UCLASS(Blueprintable, BlueprintType)
+class UPCapPropDataAsset : public UPCapDataAsset
+{
+	GENERATED_BODY()
+	
+public:
+	/** Prop DataAsset Constructor */
+	UPCapPropDataAsset();
+	
+	/** Prop Name*/
+	UPROPERTY(BlueprintReadWrite, EditAnywhere, Category="Performance Capture Prop", meta=(Tooltip="Optional Name for Prop. If blank LiveLinkSubject name will be used"))
+	FName PropName;
+	
+	/** Prop's LiveLinkSubject */
+	UPROPERTY(BlueprintReadWrite, EditDefaultsOnly, Category="Performance Capture Prop")
+	FLiveLinkSubjectName LiveLinkSubject;
+
+	/** Static Mesh Offset Transform **/
+	UPROPERTY(BlueprintReadWrite, EditDefaultsOnly, Category="Performance Capture Prop")
+	FTransform PropOffsetTransform;
+	
+	/** Prop Static - if this is left blank the mocap static mesh will be used*/
+	UPROPERTY(BlueprintReadWrite, EditDefaultsOnly, Category="Performance Capture Prop", meta=(DisplayName="Static Mesh", EditCondition="bCanSetStaticMesh", HideEditConditionToggle))
+	TSoftObjectPtr<UStaticMesh> PropStaticMesh;
+	
+	/** Prop SkelMesh - if this is left blank the mocap static mesh will be used*/
+	UPROPERTY(BlueprintReadWrite, EditAnywhere, Category="Performance Capture Prop", meta=(DisplayName="Skeletal Mesh", EditCondition="bCanSetSkeletalMesh", HideEditConditionToggle))
+	TSoftObjectPtr<USkeletalMesh> PropSkeletalMesh;
+
+	/** Custom Class - if you want your prop to be created from a custom BP actor, set the class here*/
+	UPROPERTY(BlueprintReadWrite, EditAnywhere, Category="Performance Capture Prop", meta=(DisplayName="Custom Prop Class", EditCondition="bCanSetCustomPropClass", HideEditConditionToggle))
+	TSoftClassPtr<AActor> CustomPropClass;
+
+	UPROPERTY(BlueprintReadWrite, EditDefaultsOnly, Category="Performance Capture Prop", meta=(Tooltip="Choose True for props you want to hide in 'Game' view e.g. Proxies or integrators"))
+	bool bHiddenInGame = false;
+
+	UPROPERTY(EditDefaultsOnly, Transient, Category="Performance Capture Prop", AdvancedDisplay, meta=(DisplayName="Reset Read-only Options", Tooltip="Reset all mesh options (removes read-only status on fields)"))
+	bool bClearEditConditions = false;
+
+	//~ UObject interface
+#if WITH_EDITOR
+	virtual void PostEditChangeProperty(FPropertyChangedEvent& PropertyChangedEvent) override;
+
+	virtual void PostLoad() override;
+#endif
+	//~ End UObject interface
+
+private:
+
+	UPROPERTY()
+	bool bCanSetStaticMesh = true;
+	UPROPERTY()
+	bool bCanSetSkeletalMesh = true;
+	UPROPERTY()
+	bool bCanSetCustomPropClass =true;
+
+	void ValidateEditConditions();
+};
+
+#if WITH_EDITOR
+UCLASS(Blueprintable, Abstract, meta = (ShowWorldContextPin))
+class UPerformanceCaptureDatabaseHelper : public UObject
+{
+	GENERATED_BODY()
+
+public:
+
+	UFUNCTION(BlueprintImplementableEvent)
+	void CreateRecord();
+
+	UFUNCTION(BlueprintImplementableEvent)
+	void DeleteRecord();
+};
+
+#endif
diff --git a/Engine/Plugins/VirtualProduction/PerformanceCaptureWorkflow/Source/PerformanceCaptureWorkflow/Private/PCapMVVMResolver.cpp b/Engine/Plugins/VirtualProduction/PerformanceCaptureWorkflow/Source/PerformanceCaptureWorkflow/Private/PCapMVVMResolver.cpp
new file mode 100644
--- /dev/null
+++ b/Engine/Plugins/VirtualProduction/PerformanceCaptureWorkflow/Source/PerformanceCaptureWorkflow/Private/PCapMVVMResolver.cpp
@@ -0,0 +1,29 @@
+// Copyright Epic Games, Inc. All Rights Reserved.
+
+
+#include "PCapMVVMResolver.h"
+#include "Engine/Engine.h" //Is this allowed?
+#include "PCapSubsystem.h"
+#include "PCapSettings.h"
+#include "Types/MVVMViewModelCollection.h"
+
+#if WITH_EDITOR
+
+UObject* UPCapMVVMResolver::CreateInstance(const UClass* ExpectedType, const UUserWidget* UserWidget, const UMVVMView* View) const
+{
+	const UPerformanceCaptureSettings* Settings = UPerformanceCaptureSettings::GetPerformanceCaptureSettings();
+	const UPerformanceCaptureSubsystem* Subsystem =  GEngine->GetEngineSubsystem<UPerformanceCaptureSubsystem>();
+	const UMVVMViewModelCollectionObject* Collection = Subsystem->GetViewModelCollection();
+
+	//Check the viewmodel class is valid and force load it if it's not already there.
+	if(UClass* ViewModelClass = Settings->ViewModelClass.LoadSynchronous())
+	{
+		FMVVMViewModelContext Context;
+		Context.ContextClass = ViewModelClass;
+		Context.ContextName = "PerformanceCaptureWorkflow";
+		
+		return Collection->FindViewModelInstance(Context);
+	}
+	return Super::CreateInstance(ExpectedType, UserWidget, View);
+}
+#endif
\ No newline at end of file
diff --git a/Engine/Plugins/VirtualProduction/PerformanceCaptureWorkflow/Source/PerformanceCaptureWorkflow/Private/PCapNamingTokens.cpp b/Engine/Plugins/VirtualProduction/PerformanceCaptureWorkflow/Source/PerformanceCaptureWorkflow/Private/PCapNamingTokens.cpp
new file mode 100644
--- /dev/null
+++ b/Engine/Plugins/VirtualProduction/PerformanceCaptureWorkflow/Source/PerformanceCaptureWorkflow/Private/PCapNamingTokens.cpp
@@ -0,0 +1,121 @@
+﻿// Copyright Epic Games, Inc. All Rights Reserved.
+
+
+#include "PCapNamingTokens.h"
+
+#include "PCapSessionTemplate.h"
+#include "PerformanceCapture.h"
+
+#define LOCTEXT_NAMESPACE "PCapNamingTokens"
+
+UPCapNamingTokens::UPCapNamingTokens()
+{
+	Namespace = GetPCapNamespace();
+}
+
+UPCapNamingTokens::~UPCapNamingTokens()
+{
+}
+
+void UPCapNamingTokens::OnCreateDefaultTokens(TArray<FNamingTokenData>& Tokens)
+{
+	Super::OnCreateDefaultTokens(Tokens);
+
+	Tokens.Add( {
+		TEXT("session"),
+		LOCTEXT("PCapTokenSession", "Session"),
+		FNamingTokenData::FTokenProcessorDelegateNative::CreateLambda([this]
+		{
+			if (Context && Context->SessionTemplate)
+			{
+				return FText::FromString(Context->SessionTemplate->SessionName);
+			}
+
+			UE_LOG(LogPCap, Verbose, TEXT("Attempted to evaluate token but no valid context is available."));
+			
+			return FText::GetEmpty();
+		})
+	});
+	
+	Tokens.Add( {
+		TEXT("sessionToken"),
+		LOCTEXT("PCapTokenSession", "Session"),
+		FNamingTokenData::FTokenProcessorDelegateNative::CreateLambda([this]
+		{
+			if (Context && Context->SessionTemplate)
+			{
+				return FText::FromString(Context->SessionTemplate->SessionToken.Output);
+			}
+
+			UE_LOG(LogPCap, Verbose, TEXT("Attempted to evaluate token but no valid context is available."));
+			
+			return FText::GetEmpty();
+		})
+	});
+
+	Tokens.Add( {
+		TEXT("production"),
+		LOCTEXT("PCapTokenProduction", "Production"),
+		FNamingTokenData::FTokenProcessorDelegateNative::CreateLambda([this]
+		{
+			if (Context)
+			{
+				return FText::FromString(Context->SessionTemplate->ProductionName);
+				
+			}
+
+			UE_LOG(LogPCap, Verbose, TEXT("Attempted to evaluate token but no valid context is available."));
+
+			return FText::GetEmpty();
+		})
+	});
+	
+	Tokens.Add( {
+		TEXT("pcapRootFolder"),
+		LOCTEXT("PCapTokenProduction", "RootFolder"),
+		FNamingTokenData::FTokenProcessorDelegateNative::CreateLambda([this]
+		{
+			if (Context)
+			{
+				return FText::FromString(Context->SessionTemplate->TemplateRootFolder.Path);
+			}
+
+			UE_LOG(LogPCap, Verbose, TEXT("Attempted to evaluate token but no valid context is available."));
+
+			return FText::GetEmpty();
+		})
+	});
+		
+	Tokens.Add( {
+		TEXT("sessionFolder"),
+		LOCTEXT("PCapTokenProduction", "SessionFolder"),
+		FNamingTokenData::FTokenProcessorDelegateNative::CreateLambda([this]
+		{
+			if (Context)
+			{
+				return FText::FromString(Context->SessionTemplate->SessionFolder.FolderPathOutput);
+			}
+			UE_LOG(LogPCap, Verbose, TEXT("Attempted to evaluate token but no valid context is available."));
+
+			return FText::GetEmpty();
+			
+		})
+	});
+}
+
+void UPCapNamingTokens::OnPreEvaluate_Implementation(const FNamingTokensEvaluationData& InEvaluationData)
+{
+	Super::OnPreEvaluate_Implementation(InEvaluationData);
+
+	UPCapNamingTokensContext* MatchingContext = nullptr;
+	InEvaluationData.Contexts.FindItemByClass<UPCapNamingTokensContext>(&MatchingContext);
+	Context = MatchingContext;
+}
+
+void UPCapNamingTokens::OnPostEvaluate_Implementation()
+{
+	Super::OnPostEvaluate_Implementation();
+	Context = nullptr;
+}
+
+#undef LOCTEXT_NAMESPACE
\ No newline at end of file
diff --git a/Engine/Plugins/VirtualProduction/PerformanceCaptureWorkflow/Source/PerformanceCaptureWorkflow/Private/PCapNamingTokens.h b/Engine/Plugins/VirtualProduction/PerformanceCaptureWorkflow/Source/PerformanceCaptureWorkflow/Private/PCapNamingTokens.h
new file mode 100644
--- /dev/null
+++ b/Engine/Plugins/VirtualProduction/PerformanceCaptureWorkflow/Source/PerformanceCaptureWorkflow/Private/PCapNamingTokens.h
@@ -0,0 +1,59 @@
+﻿// Copyright Epic Games, Inc. All Rights Reserved.
+
+#pragma once
+
+#include "CoreMinimal.h"
+#include "GlobalNamingTokens.h"
+#include "PCapNamingTokens.generated.h"
+
+class UPCapSessionTemplate;
+
+/**
+ * 
+ */
+UCLASS(NotBlueprintable)
+class PERFORMANCECAPTUREWORKFLOW_API UPCapNamingTokens : public UNamingTokens
+{
+	GENERATED_BODY()
+
+public:
+
+	UPCapNamingTokens();
+	virtual ~UPCapNamingTokens() override;
+	
+	static FString GetPCapNamespace()
+	{
+		return TEXT("pcap");
+	}
+
+protected:
+
+	// ~Begin UNamingTokens
+	virtual void OnCreateDefaultTokens(TArray<FNamingTokenData>& Tokens) override;
+	virtual void OnPreEvaluate_Implementation(const FNamingTokensEvaluationData& InEvaluationData) override;
+	virtual void OnPostEvaluate_Implementation() override;
+	// ~End UNamingTokens
+
+private:
+
+	UPROPERTY(Transient)
+	TObjectPtr<UPCapNamingTokensContext> Context;
+};
+
+/*
+ *Context Object that references the current session template and allows it to be passed to the naming token functions
+ * 
+ */
+
+UCLASS(MinimalAPI)
+
+class UPCapNamingTokensContext : public UObject
+{
+	GENERATED_BODY()
+	
+public:
+	
+	//The current template in use. This all seems rather circular
+	UPROPERTY(Transient)
+	TObjectPtr<const UPCapSessionTemplate> SessionTemplate;
+};
\ No newline at end of file
diff --git a/Engine/Plugins/VirtualProduction/PerformanceCaptureWorkflow/Source/PerformanceCaptureWorkflow/Private/PCapPropComponent.cpp b/Engine/Plugins/VirtualProduction/PerformanceCaptureWorkflow/Source/PerformanceCaptureWorkflow/Private/PCapPropComponent.cpp
new file mode 100644
--- /dev/null
+++ b/Engine/Plugins/VirtualProduction/PerformanceCaptureWorkflow/Source/PerformanceCaptureWorkflow/Private/PCapPropComponent.cpp
@@ -0,0 +1,269 @@
+// Copyright Epic Games, Inc. All Rights Reserved.
+
+
+#include "PCapPropComponent.h"
+#include "Components/SkeletalMeshComponent.h"
+#include "ILiveLinkClient.h"
+#include "PCapPropLiveLinkAnimInstance.h"
+#include "GameFramework/Actor.h"
+#include "Components/SceneComponent.h"
+#include "Features/IModularFeatures.h"
+#include "Roles/LiveLinkAnimationRole.h"
+#include "Roles/LiveLinkTransformRole.h"
+#include "Roles/LiveLinkTransformTypes.h"
+
+
+// Sets default values for this component's properties
+UPCapPropComponent::UPCapPropComponent()
+	: bIsDirty(true)
+{
+	PrimaryComponentTick.bCanEverTick = true;
+	PrimaryComponentTick.bStartWithTickEnabled = true;
+	PrimaryComponentTick.TickGroup = ETickingGroup::TG_PrePhysics;
+	bTickInEditor = true;
+}
+
+UPCapPropComponent::~UPCapPropComponent()
+{
+	//TODO: Might have to do something when exiting PIE here
+}
+
+USceneComponent* UPCapPropComponent::GetControlledComponent() const
+{
+	if(ControlledComponent.GetComponent(GetOwner()))
+	{
+		return Cast<USceneComponent>(ControlledComponent.GetComponent(GetOwner()));
+	}
+
+	return nullptr;
+}
+
+void UPCapPropComponent::SetControlledComponent(USceneComponent* InComponent)
+{
+	//Check the offered new component is inside the same actor as this prop component
+	if(InComponent->GetOwner() == this->GetOwner())
+	{
+		ControlledComponent.OverrideComponent = InComponent;
+	}
+}
+
+void UPCapPropComponent::SetLiveLinkSubject(FLiveLinkSubjectName Subject)
+{
+	SubjectName = Subject;
+
+	//If there is a Skeletal Mesh controlled by an Anim Instance, reinit.
+	InitiateAnimation();
+}
+
+FLiveLinkSubjectName UPCapPropComponent::GetLiveLinkSubject() const
+{
+	return SubjectName;
+}
+
+void UPCapPropComponent::SetEvaluateLiveLinkData(bool bEvaluate)
+{
+	bEvaluateLiveLink = bEvaluate;
+
+	//If there is a Skeletal Mesh controlled by an Anim Instance, reinit.
+	USkeletalMeshComponent* ControlledMesh = Cast<USkeletalMeshComponent>(ControlledComponent.GetComponent(GetOwner()));
+	if(IsValid(ControlledMesh))
+	{
+		UPCapPropLiveLinkAnimInstance* LiveLinkAnimInstance;
+		LiveLinkAnimInstance = Cast<UPCapPropLiveLinkAnimInstance>(ControlledMesh->GetAnimInstance());
+		if(LiveLinkAnimInstance)
+		{
+			LiveLinkAnimInstance->EnableLiveLinkEvaluation(bEvaluateLiveLink);
+		}
+	}
+}
+
+bool UPCapPropComponent::GetEvaluateLiveLinkData()
+{
+	return bEvaluateLiveLink;
+}
+
+void UPCapPropComponent::SetOffsetTransform(FTransform NewOffset)
+{
+	OffsetTransform = NewOffset;
+	
+	//If there is a Skeletal Mesh controlled by an Anim Instance, reinit.
+	TObjectPtr<USkeletalMeshComponent> ControlledMesh = Cast<USkeletalMeshComponent>(ControlledComponent.GetComponent(GetOwner()));
+	if(IsValid(ControlledMesh))
+	{
+		UPCapPropLiveLinkAnimInstance* LiveLinkAnimInstance;
+		LiveLinkAnimInstance = Cast<UPCapPropLiveLinkAnimInstance>(ControlledMesh->GetAnimInstance());
+		if(LiveLinkAnimInstance)
+		{
+			LiveLinkAnimInstance->SetOffsetTransform(OffsetTransform);
+		}
+	}
+}
+
+void UPCapPropComponent::OnRegister()
+{
+	Super::OnRegister();
+	bIsDirty = true;
+
+	GetOwner()->Tags.AddUnique(TEXT("PCapProp"));
+
+	//Enforce mobility on the controlled component
+	if(ControlledComponent.GetComponent(GetOwner()))
+	{
+		Cast<USceneComponent>(ControlledComponent.GetComponent(GetOwner()))->SetMobility(EComponentMobility::Movable);
+	}
+
+	TObjectPtr<USkeletalMeshComponent> ControlledMesh = Cast<USkeletalMeshComponent>(ControlledComponent.GetComponent(GetOwner()));
+	if(ControlledMesh && bIsDirty)
+	{
+		InitiateAnimation();
+	}
+}
+
+void UPCapPropComponent::DestroyComponent(bool bPromoteChildren)
+{
+	ResetAnimInstance();
+	
+	Super::DestroyComponent(bPromoteChildren);
+}
+
+#if WITH_EDITOR
+void UPCapPropComponent::PostEditChangeProperty(FPropertyChangedEvent& PropertyChangedEvent)
+{
+	if(ControlledComponent.GetComponent(GetOwner()))
+	{
+		Cast<USceneComponent>(ControlledComponent.GetComponent(GetOwner()))->SetMobility(EComponentMobility::Movable);
+	}
+	if(PropertyChangedEvent.Property->GetFName() == GET_MEMBER_NAME_CHECKED(UPCapPropComponent, OffsetTransform))
+	{
+		SetOffsetTransform(OffsetTransform);
+	}
+	if(PropertyChangedEvent.Property->GetFName() == GET_MEMBER_NAME_CHECKED(UPCapPropComponent, ControlledComponent))
+	{
+		ResetAnimInstance();
+		bIsDirty = true;
+	}
+	if(PropertyChangedEvent.Property->GetFName() == GET_MEMBER_NAME_CHECKED(UPCapPropComponent, SubjectName))
+	{
+		ResetAnimInstance();
+		bIsDirty = true;
+	}
+
+	Super::PostEditChangeProperty(PropertyChangedEvent);
+}
+#endif
+
+void UPCapPropComponent::TickComponent(float DeltaTime, ELevelTick TickType, FActorComponentTickFunction* ThisTickFunction)
+{
+
+	//Check if spawnable
+	if(!bIsSpawnableCache.IsSet() || bIsDirty)
+	{
+		static const FName SequencerActorTag(TEXT("SequencerActor"));
+		const AActor* OwningActor = GetOwner();
+
+		bIsSpawnableCache = OwningActor && OwningActor->ActorHasTag(SequencerActorTag);
+
+		if(*bIsSpawnableCache)
+		{
+			bEvaluateLiveLink = false;
+			bIsControlledBySequencer = true;
+		}
+	}
+
+	if(!bIsControlledBySequencer)
+	{
+		// If we are controlling a static mesh, follow this path of applying a relative transform to the controlled mesh from the LiveLink data
+		const TObjectPtr<UStaticMeshComponent> ControlledMesh = Cast<UStaticMeshComponent>(ControlledComponent.GetComponent(GetOwner()));
+		if(ControlledMesh)
+		{
+			FLiveLinkSubjectFrameData SubjectFrameData;
+
+			ILiveLinkClient& LiveLinkClient = IModularFeatures::Get().GetModularFeature<ILiveLinkClient>(ILiveLinkClient::ModularFeatureName);
+	
+			LiveLinkClient.EvaluateFrame_AnyThread(SubjectName, LiveLinkClient.GetSubjectRole_AnyThread(SubjectName), SubjectFrameData);
+	
+			if(LiveLinkClient.GetSubjectRole_AnyThread(SubjectName) == ULiveLinkAnimationRole::StaticClass())
+			{
+				const FLiveLinkAnimationFrameData* FrameData = SubjectFrameData.FrameData.Cast<FLiveLinkAnimationFrameData>();
+				if(FrameData && bEvaluateLiveLink)
+				{
+					CachedLiveLinkTransform = FrameData->Transforms[0];
+				}
+				GetControlledComponent()->SetRelativeTransform(OffsetTransform * CachedLiveLinkTransform, false, nullptr, ETeleportType::None);
+
+			}
+			else if(LiveLinkClient.GetSubjectRole_AnyThread(SubjectName) == ULiveLinkTransformRole::StaticClass())
+			{
+				const FLiveLinkTransformFrameData* FrameData = SubjectFrameData.FrameData.Cast<FLiveLinkTransformFrameData>();
+				if(FrameData && bEvaluateLiveLink)
+				{
+					CachedLiveLinkTransform = FrameData->Transform;
+				}
+				GetControlledComponent()->SetRelativeTransform(OffsetTransform * CachedLiveLinkTransform, false, nullptr, ETeleportType::None);
+			}
+		}
+		//If we are controlling a Skelmesh and receiving a transform role, apply this route. 
+		const TObjectPtr<USkeletalMeshComponent> ControlledSkelmesh = Cast<USkeletalMeshComponent>(ControlledComponent.GetComponent(GetOwner()));
+		if(ControlledSkelmesh)
+		{
+			FLiveLinkSubjectFrameData SkelSubjectFrameData;
+
+			ILiveLinkClient& SkelLiveLinkClient = IModularFeatures::Get().GetModularFeature<ILiveLinkClient>(ILiveLinkClient::ModularFeatureName);
+
+			SkelLiveLinkClient.EvaluateFrame_AnyThread(SubjectName, SkelLiveLinkClient.GetSubjectRole_AnyThread(SubjectName), SkelSubjectFrameData);
+			if(SkelLiveLinkClient.GetSubjectRole_AnyThread(SubjectName) == ULiveLinkTransformRole::StaticClass())
+			{
+				const FLiveLinkTransformFrameData* FrameData = SkelSubjectFrameData.FrameData.Cast<FLiveLinkTransformFrameData>();
+				if(FrameData && bEvaluateLiveLink)
+				{
+					CachedLiveLinkTransform = FrameData->Transform;
+				}
+				GetControlledComponent()->SetRelativeTransform(OffsetTransform * CachedLiveLinkTransform, false, nullptr, ETeleportType::None);
+			}
+			if((SkelLiveLinkClient.GetSubjectRole_AnyThread(SubjectName) == ULiveLinkAnimationRole::StaticClass()) && bIsDirty)
+			{
+				InitiateAnimation();
+			}
+		}
+	}
+	Super::TickComponent(DeltaTime, TickType, ThisTickFunction);
+}
+
+void UPCapPropComponent::InitiateAnimation()
+{
+	TObjectPtr<USkeletalMeshComponent> ControlledMesh = Cast<USkeletalMeshComponent>(ControlledComponent.GetComponent(GetOwner()));
+	ILiveLinkClient& LiveLinkClient = IModularFeatures::Get().GetModularFeature<ILiveLinkClient>(ILiveLinkClient::ModularFeatureName);
+	
+	//If we are controlling a Skeletal Mesh and the Live Link data is of the animation role, initialize animation instance
+	if(IsValid(ControlledMesh) && LiveLinkClient.GetSubjectRole_AnyThread(SubjectName) == ULiveLinkAnimationRole::StaticClass())
+	{
+		ControlledMesh->SetAnimInstanceClass(UPCapPropLiveLinkAnimInstance::StaticClass());
+		ControlledMesh->InitAnim(true /*bForceReinit*/);
+		ControlledMesh->SetUpdateAnimationInEditor(true);
+		ControlledMesh->bPropagateCurvesToFollowers = true;
+		ControlledMesh->VisibilityBasedAnimTickOption = EVisibilityBasedAnimTickOption::AlwaysTickPoseAndRefreshBones;
+		if(ControlledMesh)
+		{
+			UPCapPropLiveLinkAnimInstance* AnimInstance;
+			AnimInstance = Cast<UPCapPropLiveLinkAnimInstance>(ControlledMesh->GetAnimInstance());
+			if(AnimInstance)
+			{
+				AnimInstance->SetSubject(SubjectName);
+				AnimInstance->EnableLiveLinkEvaluation(bEvaluateLiveLink);
+				AnimInstance->SetOffsetTransform(OffsetTransform);
+				bIsDirty=false;
+			}
+		}
+	}
+}
+
+void UPCapPropComponent::ResetAnimInstance() const
+{
+	TObjectPtr<USkeletalMeshComponent> ControlledMesh = Cast<USkeletalMeshComponent>(ControlledComponent.GetComponent(GetOwner()));
+	if(ControlledMesh)
+	{
+		ControlledMesh->SetAnimInstanceClass(nullptr);
+		ControlledMesh->InitAnim(true /*bForceReinit*/);
+		ControlledMesh->SetUpdateAnimationInEditor(true);
+	}
+}
diff --git a/Engine/Plugins/VirtualProduction/PerformanceCaptureWorkflow/Source/PerformanceCaptureWorkflow/Private/PCapPropComponent.h b/Engine/Plugins/VirtualProduction/PerformanceCaptureWorkflow/Source/PerformanceCaptureWorkflow/Private/PCapPropComponent.h
new file mode 100644
--- /dev/null
+++ b/Engine/Plugins/VirtualProduction/PerformanceCaptureWorkflow/Source/PerformanceCaptureWorkflow/Private/PCapPropComponent.h
@@ -0,0 +1,132 @@
+// Copyright Epic Games, Inc. All Rights Reserved.
+
+#pragma once
+
+#include "CoreMinimal.h"
+#include "ILiveLinkClient.h"
+#include "Components/SceneComponent.h"
+#include "Components/ActorComponent.h"
+#include "Components/StaticMeshComponent.h"
+#include "PCapPropComponent.generated.h"
+
+/**
+* A component for use in Motion Capture props. Accepts data as either a LiveLink Transform Role or an Animation Role. If the Live Link data is in the Animation role, this component will take the root bone transform and apply that transform to the owning actor's root component. Can be applied to Static Meshes, Skeletal Meshes and Blueprint constructions. Transform data can be offset in the component's local space. If this component is driving a Skeletal Mesh component and is receiving Live Link Animation data the full bone hierarchy will be applied. 
+*/
+UCLASS(Blueprintable, ClassGroup=("Performance Capture"), meta=(BlueprintSpawnableComponent, DisplayName="Prop Component"))
+class PERFORMANCECAPTUREWORKFLOW_API UPCapPropComponent : public UActorComponent
+{
+	GENERATED_BODY()
+
+public:
+	// Constructor
+	UPCapPropComponent();
+	// Destructor
+	~UPCapPropComponent();
+
+	/**
+	 * LiveLink Subject Name. Must have either Animation Role Type.
+	 */
+	UPROPERTY(EditAnywhere, Category = "Performance Capture")
+	FLiveLinkSubjectName SubjectName;
+
+	/**
+	 * Should LiveLink Subject data be evaluated.
+	 */
+	UPROPERTY(EditAnywhere, Category = "Performance Capture")
+	bool bEvaluateLiveLink = true;
+
+	/**
+	 * Overrides all LiveLink data and offset transform. Must be manually set on possessable bindings.
+	 */
+	UPROPERTY(BlueprintReadOnly, EditInstanceOnly, Category = "Performance Capture")
+	bool bIsControlledBySequencer = false;
+
+	/**
+	 *The package name of the spawning data asset, if there is one. 
+	 */
+	UPROPERTY(BlueprintReadOnly, Category = "Performance Capture")
+	FName SpawningDataAsset;
+	
+	/**
+	 * Offset the incoming LiveLink Transform data in the local space of the controlled component.
+	 */
+	UPROPERTY(BlueprintReadOnly, EditAnywhere, Category = "Performance Capture")
+	FTransform OffsetTransform = FTransform::Identity;
+
+	/**
+	 * Scene Component driven by the LiveLink data. Defaults to Root component of the owning actor.
+	 */
+	UPROPERTY(EditInstanceOnly, Category = "Performance Capture", meta=(UseComponentPicker, AllowedClasses = "/Script/Engine.SceneComponent"))
+	FComponentReference ControlledComponent;
+
+	/**
+	 * Returns the component controlled by this component. Can return null.
+	 */
+	UFUNCTION(BlueprintPure, Category = "Performance Capture")
+	USceneComponent* GetControlledComponent() const;
+
+	/**
+	 * Sets the component controlled by this component. Component must be within the same actor as this Prop component.
+	 */
+	UFUNCTION(BlueprintCallable, Category = "Performance Capture")
+	void SetControlledComponent(USceneComponent* InComponent);
+
+	UFUNCTION(BlueprintCallable, Category = "Performance Capture|Prop")
+	void SetLiveLinkSubject(FLiveLinkSubjectName Subject);
+	
+	/**
+	* Get the LiveLink Subject Name.
+	* @return FLiveLinkSubjectName Current LiveLink Subject.
+	*/
+	UFUNCTION(BlueprintCallable, Category ="Performance Capture|Prop")
+	FLiveLinkSubjectName GetLiveLinkSubject() const;
+
+	/**
+	* Set the LiveLink data to update the Skeletal Mesh pose.
+	* @param bEvaluateLinkLink Drive or pause the Controlled Mesh from LiveLink Subject data.
+	*/
+	UFUNCTION(BlueprintCallable, Category ="Performance Capture|Prop")
+	void SetEvaluateLiveLinkData(bool bEvaluateLinkLink);
+	
+	/**
+	* Get the LiveLink Evaluation State. Subject must have the Animation Role Type.
+	* @return bool Is LiveLink data being evaluated.
+	*/
+	UFUNCTION(BlueprintPure, Category ="Performance Capture|Prop")
+	bool GetEvaluateLiveLinkData();
+
+	/**
+	* Update the local space offset to this prop's Live Link pose. 
+	* @param NewOffset Offset Transform.
+	*/
+	UFUNCTION(BlueprintCallable, Category ="Performance Capture|Prop")
+	void SetOffsetTransform(FTransform NewOffset);
+
+
+protected:
+
+	//~ Begin UActorComponent interface
+	virtual void OnRegister() override;
+	virtual void DestroyComponent(bool bPromoteChildren) override;
+	//~ End UActorComponent interface
+	
+	//~ UObject interface
+#if WITH_EDITOR
+	virtual void PostEditChangeProperty(FPropertyChangedEvent& PropertyChangedEvent) override;
+#endif
+	//~ End UObject interface
+	
+
+public:
+	// Called every frame
+	virtual void TickComponent(float DeltaTime, ELevelTick TickType, FActorComponentTickFunction* ThisTickFunction) override;
+
+	void InitiateAnimation();
+
+	void ResetAnimInstance() const;
+	
+private:
+	FTransform CachedLiveLinkTransform;
+	bool bIsDirty;
+	TOptional<bool> bIsSpawnableCache;
+};
diff --git a/Engine/Plugins/VirtualProduction/PerformanceCaptureWorkflow/Source/PerformanceCaptureWorkflow/Private/PCapPropLiveLinkAnimInstance.cpp b/Engine/Plugins/VirtualProduction/PerformanceCaptureWorkflow/Source/PerformanceCaptureWorkflow/Private/PCapPropLiveLinkAnimInstance.cpp
new file mode 100644
--- /dev/null
+++ b/Engine/Plugins/VirtualProduction/PerformanceCaptureWorkflow/Source/PerformanceCaptureWorkflow/Private/PCapPropLiveLinkAnimInstance.cpp
@@ -0,0 +1,55 @@
+// Copyright Epic Games, Inc. All Rights Reserved.
+
+
+#include "PCapPropLiveLinkAnimInstance.h"
+
+#include UE_INLINE_GENERATED_CPP_BY_NAME(PCapPropLiveLinkAnimInstance)
+
+/*------------------------------------------------------------------------------
+	Anim Instance Proxy.
+------------------------------------------------------------------------------*/
+
+
+void FPCapPropLiveLinkAnimInstanceProxy::Initialize(UAnimInstance* InAnimInstance)
+{
+	FAnimInstanceProxy::Initialize(InAnimInstance);
+	FAnimationInitializeContext InitContext(this);
+	PoseNode.Initialize_AnyThread(InitContext);
+}
+
+void FPCapPropLiveLinkAnimInstanceProxy::PreUpdate(UAnimInstance* InAnimInstance, float DeltaSeconds)
+{
+	Super::PreUpdate(InAnimInstance, DeltaSeconds);
+
+	if(PoseNode.HasPreUpdate())
+	{
+		PoseNode.PreUpdate(InAnimInstance);
+	}
+}
+
+bool FPCapPropLiveLinkAnimInstanceProxy::Evaluate(FPoseContext& Output)
+{
+	PoseNode.Evaluate_AnyThread(Output);
+	
+	return true;
+}
+
+void FPCapPropLiveLinkAnimInstanceProxy::UpdateAnimationNode(const FAnimationUpdateContext& InContext)
+{
+	UpdateCounter.Increment();
+
+	PoseNode.Update_AnyThread(InContext);
+}
+
+/*------------------------------------------------------------------------------
+	Anim Instance
+------------------------------------------------------------------------------*/
+FAnimInstanceProxy* UPCapPropLiveLinkAnimInstance::CreateAnimInstanceProxy()
+{
+	return new FPCapPropLiveLinkAnimInstanceProxy(this);
+}
+
+void UPCapPropLiveLinkAnimInstance::DestroyAnimInstanceProxy(FAnimInstanceProxy* InProxy)
+{
+	Super::DestroyAnimInstanceProxy(InProxy);
+}
diff --git a/Engine/Plugins/VirtualProduction/PerformanceCaptureWorkflow/Source/PerformanceCaptureWorkflow/Private/PCapPropLiveLinkAnimInstance.h b/Engine/Plugins/VirtualProduction/PerformanceCaptureWorkflow/Source/PerformanceCaptureWorkflow/Private/PCapPropLiveLinkAnimInstance.h
new file mode 100644
--- /dev/null
+++ b/Engine/Plugins/VirtualProduction/PerformanceCaptureWorkflow/Source/PerformanceCaptureWorkflow/Private/PCapPropLiveLinkAnimInstance.h
@@ -0,0 +1,76 @@
+// Copyright Epic Games, Inc. All Rights Reserved.
+
+#pragma once
+
+#include "AnimNode_LiveLinkProp.h"
+#include "Animation/AnimInstanceProxy.h"
+#include "Animation/AnimInstance.h"
+#include "PCapPropLiveLinkAnimInstance.generated.h"
+
+
+USTRUCT()
+struct FPCapPropLiveLinkAnimInstanceProxy : public FAnimInstanceProxy
+{
+public:
+	friend struct FAnimNode_LiveLinkProp;
+
+	GENERATED_BODY()
+
+	FPCapPropLiveLinkAnimInstanceProxy()
+	{
+	}
+
+	FPCapPropLiveLinkAnimInstanceProxy(UAnimInstance* InAnimInstance)
+		: FAnimInstanceProxy(InAnimInstance)
+	{
+	}
+
+	virtual void Initialize(UAnimInstance* InAnimInstance) override;
+	virtual void PreUpdate(UAnimInstance* InAnimInstance, float DeltaSeconds) override;
+	virtual bool Evaluate(FPoseContext& Output) override;
+	virtual void UpdateAnimationNode(const FAnimationUpdateContext& InContext) override;
+
+	UPROPERTY(EditAnywhere, Category = Settings)
+	FAnimNode_LiveLinkProp PoseNode;
+};
+/**
+ * 
+ */
+UCLASS(transient, NotBlueprintable)
+class PERFORMANCECAPTUREWORKFLOW_API UPCapPropLiveLinkAnimInstance : public UAnimInstance
+{
+	GENERATED_BODY()
+
+public:
+
+	UFUNCTION(BlueprintCallable, Category = "Performance Capture|Prop")
+	void SetSubject(FLiveLinkSubjectName SubjectName)
+	{
+		GetProxyOnGameThread<FPCapPropLiveLinkAnimInstanceProxy>().PoseNode.LiveLinkSubjectName = SubjectName;
+	}
+
+	UFUNCTION(BlueprintCallable, Category = "Performance Capture|Prop")
+	void EnableLiveLinkEvaluation(bool bDoEnable)
+	{
+		GetProxyOnGameThread<FPCapPropLiveLinkAnimInstanceProxy>().PoseNode.bDoLiveLinkEvaluation = bDoEnable;
+	}
+	
+	UFUNCTION(BlueprintCallable, Category = "Performance Capture|Prop")
+	void SetOffsetTransform(FTransform Offset)
+	{
+		GetProxyOnGameThread<FPCapPropLiveLinkAnimInstanceProxy>().PoseNode.OffsetTransform = Offset;
+	}
+	
+	//Not exposed to Blueprint
+	bool GetEnableLiveLinkEvaluation()
+	{
+		return GetProxyOnGameThread<FPCapPropLiveLinkAnimInstanceProxy>().PoseNode.bDoLiveLinkEvaluation;
+	}
+
+	protected:
+	virtual FAnimInstanceProxy* CreateAnimInstanceProxy() override;
+	virtual void DestroyAnimInstanceProxy(FAnimInstanceProxy* InProxy) override;
+
+	friend FPCapPropLiveLinkAnimInstanceProxy;
+	
+};
diff --git a/Engine/Plugins/VirtualProduction/PerformanceCaptureWorkflow/Source/PerformanceCaptureWorkflow/Private/PCapSessionTemplate.cpp b/Engine/Plugins/VirtualProduction/PerformanceCaptureWorkflow/Source/PerformanceCaptureWorkflow/Private/PCapSessionTemplate.cpp
new file mode 100644
--- /dev/null
+++ b/Engine/Plugins/VirtualProduction/PerformanceCaptureWorkflow/Source/PerformanceCaptureWorkflow/Private/PCapSessionTemplate.cpp
@@ -0,0 +1,111 @@
+﻿// Copyright Epic Games, Inc. All Rights Reserved.
+
+
+#include "PCapSessionTemplate.h"
+#include "NamingTokensEngineSubsystem.h"
+#include "PCapBPFunctionLibrary.h"
+#include "Engine/Engine.h"
+
+UPCapSessionTemplate::UPCapSessionTemplate(const FObjectInitializer& ObjectInitializer)
+{
+	bIsEditable = true;
+	NamingTokensContext = NewObject<UPCapNamingTokensContext>();
+	
+}
+
+void UPCapSessionTemplate::PostLoad()
+{
+	Super::PostLoad();
+	
+}
+
+void UPCapSessionTemplate::PostEditChangeProperty(FPropertyChangedEvent& PropertyChangedEvent)
+{
+	Super::PostEditChangeProperty(PropertyChangedEvent);
+
+	if (PropertyChangedEvent.Property != nullptr)
+	{
+		UpdateAllFields();
+	}
+}
+
+void UPCapSessionTemplate::UpdateAllFields()
+{
+
+	if (bIsEditable) //If this asset is "locked" then do no update any of the strings. 
+	{
+		//Make sure all the tokens are updated each time a property is edited
+		SessionToken		= UpdateStringTemplate(SessionToken);
+	
+		SessionFolder		= UpdateFolderPathTemplate(SessionFolder);
+		CommonFolder		= UpdateFolderPathTemplate(CommonFolder);
+		CharacterFolder		= UpdateFolderPathTemplate(CharacterFolder);
+		PerformerFolder		= UpdateFolderPathTemplate(PerformerFolder);
+		PropFolder			= UpdateFolderPathTemplate(PropFolder);
+		SceneFolder			= UpdateFolderPathTemplate(SceneFolder);
+		TakeFolder			= UpdateFolderPathTemplate(TakeFolder);
+
+		for (TPair<FName, FPCapTokenisedFolderPath>& Pair : AdditionalFolders)
+		{
+			Pair.Value = UpdateFolderPathTemplate(Pair.Value);
+		}
+
+	
+		TakeSaveName		= UpdateStringTemplate(TakeSaveName);
+		AnimationTrackName	= UpdateStringTemplate(AnimationTrackName);
+		AnimationAssetName	= UpdateStringTemplate(AnimationAssetName);
+		AnimationSubDirectory= UpdateStringTemplate(AnimationSubDirectory);
+		AudioSourceName		= UpdateStringTemplate(AudioSourceName);
+		AudioTrackName		= UpdateStringTemplate(AudioTrackName);
+		AudioAssetName		= UpdateStringTemplate(AudioAssetName);
+		AudioSubDirectory	= UpdateStringTemplate(AudioSubDirectory);
+	}
+}
+
+FPCapTokenisedString UPCapSessionTemplate::UpdateStringTemplate(FPCapTokenisedString& InTokenisedTemplate)
+{
+	FPCapTokenisedString Template = InTokenisedTemplate;
+	
+	if (!NamingTokensContext)
+	{
+		NamingTokensContext = NewObject<UPCapNamingTokensContext>();
+	}
+
+	if (bIsEditable && NamingTokensContext)
+	{
+		NamingTokensContext->SessionTemplate = this;
+
+		FNamingTokenFilterArgs NamingTokenFilters;
+		NamingTokenFilters.NamespacesToInclude = { FString("pcap") , FString("tr")};
+
+		FNamingTokenResultData NamingTokenResultData = GEngine->GetEngineSubsystem<UNamingTokensEngineSubsystem>()->EvaluateTokenText(FText::FromString(Template.Template), NamingTokenFilters, {NamingTokensContext});
+		
+		Template.Output = UPerformanceCaptureBPFunctionLibrary::SanitizeFileString(NamingTokenResultData.EvaluatedText.ToString());
+	}
+
+	return Template;
+	
+}
+
+FPCapTokenisedFolderPath UPCapSessionTemplate::UpdateFolderPathTemplate(FPCapTokenisedFolderPath& InFolderPathTokenisedTemplate)
+{
+	FPCapTokenisedFolderPath Template = InFolderPathTokenisedTemplate;
+	if (!NamingTokensContext)
+	{
+		NamingTokensContext = NewObject<UPCapNamingTokensContext>();
+	}
+
+	if (bIsEditable && NamingTokensContext)
+	{
+		NamingTokensContext->SessionTemplate = this;
+
+		FNamingTokenFilterArgs NamingTokenFilters;
+		NamingTokenFilters.NamespacesToInclude = { FString("pcap") };
+
+		FNamingTokenResultData NamingTokenResultData = GEngine->GetEngineSubsystem<UNamingTokensEngineSubsystem>()->EvaluateTokenText(FText::FromString(Template.FolderPathTemplate), NamingTokenFilters, {NamingTokensContext});
+		
+		Template.FolderPathOutput = UPerformanceCaptureBPFunctionLibrary::SanitizePathString(NamingTokenResultData.EvaluatedText.ToString());
+	}
+
+	return Template;
+}
diff --git a/Engine/Plugins/VirtualProduction/PerformanceCaptureWorkflow/Source/PerformanceCaptureWorkflow/Private/PCapSessionTemplate.h b/Engine/Plugins/VirtualProduction/PerformanceCaptureWorkflow/Source/PerformanceCaptureWorkflow/Private/PCapSessionTemplate.h
new file mode 100644
--- /dev/null
+++ b/Engine/Plugins/VirtualProduction/PerformanceCaptureWorkflow/Source/PerformanceCaptureWorkflow/Private/PCapSessionTemplate.h
@@ -0,0 +1,262 @@
+﻿// Copyright Epic Games, Inc. All Rights Reserved.
+
+#pragma once
+
+#include "CoreMinimal.h"
+#include "PCapDatabase.h"
+#include "PCapNamingTokens.h"
+#include "PCapSessionTemplate.generated.h"
+
+/**
+ * Struct for handling conversion of token entry template into a validated folder path
+ */
+
+USTRUCT(BlueprintType)
+struct FPCapTokenisedString
+{
+	GENERATED_BODY()
+
+	FPCapTokenisedString() = default;
+
+	UPROPERTY(EditAnywhere, BlueprintReadOnly, Category="Folder Template")
+	FString Template;
+
+	UPROPERTY(VisibleAnywhere, BlueprintReadOnly, Category="Folder Template")
+	FString Output;
+	
+};
+
+USTRUCT(BlueprintType)
+struct FPCapTokenisedFolderPath
+{
+	GENERATED_BODY()
+
+	FPCapTokenisedFolderPath() = default;
+
+	UPROPERTY(EditAnywhere, BlueprintReadOnly, Category="Folder Template")
+	FString FolderPathTemplate;
+
+	UPROPERTY(VisibleAnywhere, BlueprintReadOnly, Category="Folder Template")
+	FString FolderPathOutput;
+	
+};
+
+/**
+ * Data asset for defining all the folders that will be created for a session
+ * Can be duplicated and locked to ensure all tokenized strings become serialized static data.
+ */
+
+UCLASS(Blueprintable, BlueprintType)
+class UPCapSessionTemplate : public UPCapDataAsset
+{
+	GENERATED_BODY()
+
+	UPCapSessionTemplate(const FObjectInitializer& ObjectInitializer);
+
+public:
+
+	/*
+	* Core session fields
+	*/
+
+	UPROPERTY(EditAnywhere, BlueprintReadOnly, Category="Template", meta = (DisplayName = "Root Folder", ContentDir, EditCondition="bIsEditable"))
+	FDirectoryPath TemplateRootFolder;
+
+	/*
+	 *These properties are private because we want to force the use of the Blueprint Getter and Setter so modifications only follow that path.
+	 */
+
+	virtual void PostLoad() override;
+	
+	UPROPERTY(EditAnywhere, Category="Template", BlueprintGetter = GetProductionName, BlueprintSetter = SetProductionName,  meta = (DisplayName = "Production Name", EditCondition="bIsEditable"))
+	FString ProductionName;
+
+	UPROPERTY(EditAnywhere, Category="Template", BlueprintGetter = GetSessionName, BlueprintSetter = SetSessionName, meta = (DisplayName = "Session Name", EditCondition="bIsEditable"))
+	FString SessionName;
+
+
+		
+	UPROPERTY(EditAnywhere, BlueprintReadOnly, Category="Template", meta = (DisplayName = "Session Token", ContentDir, EditCondition="bIsEditable", tooltip="The final name of our session, including any tokens you include in it's construction. You can use tokens to build this string, including the Session Name above"))
+	FPCapTokenisedString SessionToken;
+
+	UFUNCTION(BlueprintGetter)
+	FString GetSessionName()
+	{
+		return SessionName;
+	}
+
+	UFUNCTION(BlueprintSetter)
+	void SetSessionName(FString NewSessionName)
+	{
+		SessionName = NewSessionName;
+		UpdateAllFields();
+	}
+
+	UFUNCTION(BlueprintGetter)
+	FString GetProductionName()
+	{
+		return ProductionName;
+	}
+
+	UFUNCTION(BlueprintSetter)
+	void SetProductionName(FString NewProductionName)
+	{
+		ProductionName = NewProductionName;
+		UpdateAllFields();
+	}
+	
+
+	/*------------------------------------------------------------------------------
+	Folder template definitions - note these hard-coded and "opinionated" about how users will work. 
+	------------------------------------------------------------------------------*/
+
+	
+	/*
+	Session folder template
+	*/
+	UPROPERTY(EditAnywhere, BlueprintReadOnly, Category="Folders", DisplayName = "Session Folder", meta=(EditCondition="bIsEditable"))
+	FPCapTokenisedFolderPath SessionFolder;
+
+	/**Common folder template**/
+	UPROPERTY(EditAnywhere, BlueprintReadOnly, Category="Folders|Session Subfolders", meta = (DisplayName = "Common Folder", ContentDir, EditCondition="bIsEditable"))
+	FPCapTokenisedFolderPath CommonFolder;
+
+	/**Character folder template*/	
+	UPROPERTY(EditAnywhere, BlueprintReadOnly, Category="Folders|Session Subfolders", meta = (DisplayName = "Character Folder", ContentDir, EditCondition="bIsEditable"))
+	FPCapTokenisedFolderPath CharacterFolder;
+
+	/**Performer folder template*/
+	UPROPERTY(EditAnywhere, BlueprintReadOnly, Category="Folders|Session Subfolders", meta = (DisplayName = "Performer Folder", ContentDir, EditCondition="bIsEditable"))
+	FPCapTokenisedFolderPath PerformerFolder;
+
+	/**Prop folder template*/
+	UPROPERTY(EditAnywhere, BlueprintReadOnly, Category="Folders|Session Subfolders", meta = (DisplayName = "Prop Folder", ContentDir, EditCondition="bIsEditable"))
+	FPCapTokenisedFolderPath PropFolder;
+
+	/**Scene data folder template*/
+	UPROPERTY(EditAnywhere, BlueprintReadOnly, Category="Folders|Session Subfolders", meta = (DisplayName = "Scene Folder", ContentDir, EditCondition="bIsEditable"))
+	FPCapTokenisedFolderPath SceneFolder;
+
+	/**Take Recorder folder template*/
+	UPROPERTY(EditAnywhere, BlueprintReadOnly, Category="Folders|Session Subfolders", meta = (DisplayName = "Takes Folder", ContentDir, EditCondition="bIsEditable"))
+	FPCapTokenisedFolderPath TakeFolder;
+	
+	UPROPERTY(EditAnywhere, BlueprintReadOnly, Category="Folders|Session Subfolders", meta=(DisplayAfter="Additional Folders", EditCondition="bIsEditable", ToolTip="Use the additional folders for any extra data you want contain in your session - eg audio and facial capture data"))
+	TMap<FName, FPCapTokenisedFolderPath> AdditionalFolders;
+	
+	/*
+	 * Take Record related fields
+	 */
+
+	UPROPERTY(EditAnywhere, BlueprintReadOnly, Category="Take Recorder", meta = (DisplayName = "Record Timecode", EditCondition="bIsEditable"))
+	bool bRecordTimecode = true;
+
+	UPROPERTY(EditAnywhere, BlueprintReadOnly, Category="Take Recorder", meta = (DisplayName = "Record to Possessable", EditCondition="bIsEditable"))
+	bool bRecordPossessable = false;
+
+	UPROPERTY(EditAnywhere, BlueprintReadOnly, Category="Take Recorder", meta = (DisplayName = "Start at Current Timecode", EditCondition="bIsEditable"))
+	bool bStartAtCurrentTimecode = true;
+
+	UPROPERTY(EditAnywhere, BlueprintReadOnly, Category="Take Recorder", meta = (DisplayName = "Record to Subsequences", EditCondition="bIsEditable"))
+	bool bRecordSubscenes = true;
+
+	UPROPERTY(EditAnywhere, BlueprintReadOnly, Category="Take Recorder", DisplayName = "Take Save Name", meta=(EditCondition="bIsEditable",  tooltip ="This be used by take recorder as the name of the recorded level sequence and the folder for any sub-scenes and assets"))
+	FPCapTokenisedString TakeSaveName;
+
+	UPROPERTY(EditAnywhere, BlueprintReadOnly, Category="Take Recorder|Animation", DisplayName = "Animation Track Name", meta=(EditCondition="bIsEditable"))
+	FPCapTokenisedString AnimationTrackName;
+
+	UPROPERTY(EditAnywhere, BlueprintReadOnly, Category="Take Recorder|Animation", DisplayName = "Animation Asset Name", meta=(EditCondition="bIsEditable"))
+	FPCapTokenisedString AnimationAssetName;
+
+	UPROPERTY(EditAnywhere, BlueprintReadOnly, Category="Take Recorder|Animation", DisplayName = "Animation Sub Directory", meta=(EditCondition="bIsEditable"))
+	FPCapTokenisedString AnimationSubDirectory;
+
+	UPROPERTY(EditAnywhere, BlueprintReadOnly, Category="Take Recorder|Audio", DisplayName = "Audio Source Name", meta=(EditCondition="bIsEditable"))
+	FPCapTokenisedString AudioSourceName;
+
+	UPROPERTY(EditAnywhere, BlueprintReadOnly, Category="Take Recorder|Audio", DisplayName = "Audio Track Name", meta=(EditCondition="bIsEditable"))
+	FPCapTokenisedString AudioTrackName;
+
+	UPROPERTY(EditAnywhere, BlueprintReadOnly, Category="Take Recorder|Audio", DisplayName = "Audio Asset Name", meta=(EditCondition="bIsEditable"))
+	FPCapTokenisedString AudioAssetName;
+
+	UPROPERTY(EditAnywhere, BlueprintReadOnly, Category="Take Recorder|Audio", DisplayName = "Audio Sub Directory", meta=(EditCondition="bIsEditable"))
+	FPCapTokenisedString AudioSubDirectory;
+	/*
+	 * Read only fields to show the user what tokens are available to them in the pcap, global, and take recorder namespaces
+	 */
+	
+	UPROPERTY(VisibleAnywhere, Category="Tokens", meta= (displayname = "PCap Tokens", tooltip="Use these  tokens in your template strings. The namespace for these tokens is 'pcap.'")) // this is for display purposes only, not to be used in code.
+	FString PCapTokens =
+	"{session} - Name for your Performance Capture session\n"
+	"{production} - Name of Performance Capture production\n"
+	"{sessionToken} - Output value of the Session Token field\n"
+	"{pcapRootFolder} - Root folder for all Performance Capture productions and sessions\n"
+	"{sessionFolder} - folder path for Performance Capture session"
+	/*"\n"
+	"\n"
+	"\n"*/;
+
+	UPROPERTY(VisibleAnywhere, Category="Tokens", meta= (displayname = "Global Tokens", tooltip="Use these  tokens in your template strings.")) // this is for display purposes only, not to be used in code.
+	FString GlobalTokens =
+	"{yyyy} - Year (4 digit)\n"
+	"{yy} - Year (2 digit)\n"
+	"{Mmm} - 3-character Month (Pascal Case)\n"
+	"{MMM} - 3-character Month (UPPERCASE)\n"
+	"{mmm} - 3-character Month (lowercase)\n"
+	"{mm} - Month (2 digit)\n"
+	"{Ddd} - 3-character Day (Pascal Case)\n"
+	"{DDD} - 3-character Day (UPPERCASE)\n"
+	"{ddd} - 3-character Day (lowercase)\n"
+	"{dd} - Day (2 digit)\n"
+	"{ampm} - am or pm (lowercase)\n"
+	"{AMPM} - AM or PM (UPPERCASE)n"
+	"{12h} - Hour (12)\n"
+	"{24h} - Hour (24)\n"
+	"{24h} - Hour (24)\n"
+	"{24h} - Hour (24)\n"
+	"{sec} - Second\n"
+	"{ms} - Millisecond";
+
+	UPROPERTY(VisibleAnywhere, Category="Tokens", meta= (displayname = "Take Recorder Tokens", tooltip="Use these  tokens in your template strings. The namespace for these tokens is 'pcap.'")) // this is for display purposes only, not to be used in code.
+	FString TakeRecorderTokens =
+	"{day}\n"
+	"{month}\n"
+	"{year}\n"
+	"{hour}\n"
+	"{minute}\n"
+	"{second}\n"
+	"{take}\n"
+	"{slate}\n"
+	"{map}\n"
+	"{actor}\n"
+	"{channel} - audio channel. Only available for audio recording tracks";
+	
+protected:
+
+	/*
+	 * Use this to effectively "lock" the asset and prevent any dynamic tokens from being
+	 * Behind the scenes, this prevents PostEditUpdates from modifying data such that makes all
+	 * the strings and folder paths become static
+	 */
+	UPROPERTY(BlueprintReadWrite, VisibleAnywhere, Category="Admin")
+	bool bIsEditable;
+	
+#if WITH_EDITOR
+	virtual void PostEditChangeProperty(FPropertyChangedEvent& PropertyChangedEvent) override;
+#endif
+
+	UFUNCTION(BlueprintCallable, Category="Performance Capture|Session Template")
+	void UpdateAllFields();
+
+	FPCapTokenisedString UpdateStringTemplate(FPCapTokenisedString& TokenisedTemplate);
+
+	FPCapTokenisedFolderPath UpdateFolderPathTemplate(FPCapTokenisedFolderPath& FolderPathTokenisedTemplate);
+	
+private:
+
+	UPROPERTY()
+	TObjectPtr<UPCapNamingTokensContext> NamingTokensContext;
+};
+
diff --git a/Engine/Plugins/VirtualProduction/PerformanceCaptureWorkflow/Source/PerformanceCaptureWorkflow/Private/PCapSettings.cpp b/Engine/Plugins/VirtualProduction/PerformanceCaptureWorkflow/Source/PerformanceCaptureWorkflow/Private/PCapSettings.cpp
new file mode 100644
--- /dev/null
+++ b/Engine/Plugins/VirtualProduction/PerformanceCaptureWorkflow/Source/PerformanceCaptureWorkflow/Private/PCapSettings.cpp
@@ -0,0 +1,62 @@
+// Copyright Epic Games, Inc. All Rights Reserved.
+
+#include "PCapSettings.h"
+#include "ISettingsModule.h"
+#include "PCapDatabase.h"
+#include "PCapDataTable.h"
+#include "Modules/ModuleManager.h"
+
+#define LOCTEXT_NAMESPACE "UPerformanceCaptureSettings"
+
+namespace PerformanceCaptureSettings
+{
+	
+}
+
+UPerformanceCaptureSettings* UPerformanceCaptureSettings::GetPerformanceCaptureSettings()
+{
+	return GetMutableDefault<UPerformanceCaptureSettings>();
+}
+
+/** Summon the Project settings for the Performance Capture module*/
+void UPerformanceCaptureSettings::ShowPerformanceCaptureProjectSettings()
+{
+	ISettingsModule& SettingsModule = FModuleManager::LoadModuleChecked<ISettingsModule>("Settings");
+	SettingsModule.ShowViewer("Project", "Plugins","PerformanceCapture");
+}
+
+#if WITH_EDITOR
+void UPerformanceCaptureSettings::SetSessionTable(TSoftObjectPtr<UPCapDataTable> NewDataTable)
+{
+	SessionTable = NewDataTable;
+
+	FProperty* SessionTableProperty = FindFieldChecked<FProperty>(UPerformanceCaptureSettings::StaticClass(),GET_MEMBER_NAME_CHECKED(UPerformanceCaptureSettings, SessionTable));
+	FPropertyChangedEvent  PCapSessionTableChangedEvent = FPropertyChangedEvent(SessionTableProperty);
+	PostEditChangeProperty(PCapSessionTableChangedEvent);
+	TryUpdateDefaultConfigFile();
+	SaveConfig();
+}
+
+void UPerformanceCaptureSettings::SetProductionTable(TSoftObjectPtr<UPCapDataTable> NewDataTable)
+{
+	ProductionTable = NewDataTable;
+	FProperty* ProductionTableProperty = FindFieldChecked<FProperty>(UPerformanceCaptureSettings::StaticClass(),GET_MEMBER_NAME_CHECKED(UPerformanceCaptureSettings, ProductionTable));
+	FPropertyChangedEvent  PCapProductionTableChangedEvent = FPropertyChangedEvent(ProductionTableProperty);
+	PostEditChangeProperty(PCapProductionTableChangedEvent);
+	TryUpdateDefaultConfigFile();
+	SaveConfig();
+}
+
+
+void UPerformanceCaptureSettings::PostEditChangeProperty(FPropertyChangedEvent& PropertyChangedEvent)
+{
+	Super::PostEditChangeProperty(PropertyChangedEvent);
+
+	if(PropertyChangedEvent.Property !=nullptr)
+	{
+		OnPCapSettingsChanged.Broadcast();
+	}
+}
+#endif
+
+#undef LOCTEXT_NAMESPACE
diff --git a/Engine/Plugins/VirtualProduction/PerformanceCaptureWorkflow/Source/PerformanceCaptureWorkflow/Private/PCapSettings.h b/Engine/Plugins/VirtualProduction/PerformanceCaptureWorkflow/Source/PerformanceCaptureWorkflow/Private/PCapSettings.h
new file mode 100644
--- /dev/null
+++ b/Engine/Plugins/VirtualProduction/PerformanceCaptureWorkflow/Source/PerformanceCaptureWorkflow/Private/PCapSettings.h
@@ -0,0 +1,81 @@
+// Copyright Epic Games, Inc. All Rights Reserved.
+
+#pragma once
+
+#include "CoreMinimal.h"
+#include "EditorUtilityWidgetBlueprint.h"
+#include "PCapStageRoot.h"
+#include "PCapSubsystem.h"
+#include "MVVMViewModelBase.h"
+#include "UObject/Object.h"
+#include "PCapSettings.generated.h"
+
+class APerformanceCaptureStageRoot;
+class UPCapSessionTemplate;
+class UDataTable;
+class UPCapDataTable;
+class USkeletalMesh;
+
+/*
+ * Per Project Settings for Performance Capture.
+ */
+UCLASS(Config=PerformanceCaptureWorkflow, DefaultConfig, DisplayName="Performance Capture")
+
+class PERFORMANCECAPTUREWORKFLOW_API UPerformanceCaptureSettings : public UObject
+{
+	GENERATED_BODY()
+public:
+	
+	UFUNCTION(BlueprintPure, Category="Performance Capture|Settings")
+	static UPerformanceCaptureSettings* GetPerformanceCaptureSettings();
+
+	UFUNCTION(BlueprintCallable, Category = "Performance Capture|Settings")
+	void ShowPerformanceCaptureProjectSettings();
+	
+	/**Stage Root Actor to spawn for Performance Capture. Class must derive from APerformanceCaptureStageRoot*/
+	UPROPERTY(BlueprintReadOnly, EditDefaultsOnly, Config, NoClear, Category="Actors", DisplayName= "Stage Root Actor Class" ,meta = (MetaClass = "/Script/PerformanceCaptureWorkflow.PerformanceCaptureStageRoot"))
+	TSoftClassPtr<APerformanceCaptureStageRoot> StageRoot;
+
+	///
+	/**Base skeletal mesh for Mocap performers*/
+	UPROPERTY(BlueprintReadOnly, EditDefaultsOnly, Config, NoClear, Category="Actors", DisplayName= "Default Performer Mesh")
+	TSoftObjectPtr<USkeletalMesh> DefaultPerformerSkelMesh;
+
+	/** Blueprint Viewmodel class. If you change class you will need to restart the editor to pick up the change*/
+	UPROPERTY(EditAnywhere, BlueprintReadOnly, Config, NoClear, Category="UI", meta = (MetaClass = "/Script/ModelViewViewModel.MVVMViewModelBase",ConfigRestartRequired=true))
+	TSoftClassPtr<UMVVMViewModelBase> ViewModelClass;
+
+	/** Editor Utility Widget class that will be used for the Mocap Manager UI*/
+	UPROPERTY(EditAnywhere, BlueprintReadOnly, Config, Category="UI", meta = (MetaClass = "/Script/Blutility.EditorUtilityWidgetBlueprint"))
+	TSoftObjectPtr<UEditorUtilityWidgetBlueprint> MocapManagerUI;
+	
+	UPROPERTY(BlueprintReadOnly, Config, EditDefaultsOnly, Category="Database", DisplayName = "Default Session Template", meta = ( tooltip="This Asset defines the folder structure and settings used in a performance capture session"))
+	TSoftObjectPtr<UPCapSessionTemplate> DefaultSessionTemplate;
+	
+	/** Blueprint helper class for making database calls. If you change class you will need to restart the editor to pick up the change*/
+	UPROPERTY(EditAnywhere, Config, NoClear, Category="Database", DisplayName = "Database Helper Class", meta=(ConfigRestartRequired=true))
+	TSoftClassPtr<UPerformanceCaptureDatabaseHelper> DatabaseHelperClass;
+	
+	UPROPERTY(BlueprintReadOnly, Config, EditDefaultsOnly, Category = "Database|Internal Database Tables", DisplayName = "Session DataTable", meta = (RequiredAssetDataTags = "RowStructure=/Script/PerformanceCaptureWorkflow.PCapSessionRecord"))
+	TSoftObjectPtr<UPCapDataTable> SessionTable;
+
+	UPROPERTY(BlueprintReadOnly, Config, EditDefaultsOnly, Category="Database|Internal Database Tables", DisplayName = "Production DataTable", meta = (RequiredAssetDataTags = "RowStructure=/Script/PerformanceCaptureWorkflow.PCapProductionRecord"))
+	TSoftObjectPtr<UPCapDataTable> ProductionTable;
+
+	UFUNCTION(BlueprintCallable, Category="Performance Capture|Settings")
+	void SetSessionTable(TSoftObjectPtr<UPCapDataTable> NewDataTable);
+
+	UFUNCTION(BlueprintCallable, Category="Performance Capture|Settings")
+	void SetProductionTable(TSoftObjectPtr<UPCapDataTable> NewDataTable); 
+
+#if WITH_EDITOR
+	
+	DECLARE_DYNAMIC_MULTICAST_DELEGATE(FOnPCapSettingsChanged);
+	
+	UPROPERTY(BlueprintAssignable, Category="Performance Capture|Settings")
+	FOnPCapSettingsChanged OnPCapSettingsChanged;
+	
+	virtual void PostEditChangeProperty(FPropertyChangedEvent& PropertyChangedEvent) override;
+	
+#endif
+};
diff --git a/Engine/Plugins/VirtualProduction/PerformanceCaptureWorkflow/Source/PerformanceCaptureWorkflow/Private/PCapStageRoot.cpp b/Engine/Plugins/VirtualProduction/PerformanceCaptureWorkflow/Source/PerformanceCaptureWorkflow/Private/PCapStageRoot.cpp
new file mode 100644
--- /dev/null
+++ b/Engine/Plugins/VirtualProduction/PerformanceCaptureWorkflow/Source/PerformanceCaptureWorkflow/Private/PCapStageRoot.cpp
@@ -0,0 +1,20 @@
+// Copyright Epic Games, Inc. All Rights Reserved.
+
+
+#include "PCapStageRoot.h"
+#include "Components/SceneCaptureComponent2D.h"
+
+// Sets default values
+APerformanceCaptureStageRoot::APerformanceCaptureStageRoot(const FObjectInitializer& ObjectInitializer)
+	: Super(ObjectInitializer)
+{
+	// Set this actor to call Tick() every frame.  You can turn this off to improve performance if you don't need it.
+	PrimaryActorTick.bCanEverTick = true;
+
+	USceneComponent* SceneRoot = CreateDefaultSubobject<USceneComponent>("Root");
+	SetRootComponent(SceneRoot);
+	MapCaptureComponent = CreateDefaultSubobject<USceneCaptureComponent2D>("MapCaptureComponent");
+	MapCaptureComponent->SetupAttachment(GetRootComponent());
+	MapCaptureComponent->bCaptureEveryFrame = false;
+	MapCaptureComponent->bCaptureOnMovement = true;
+}
\ No newline at end of file
diff --git a/Engine/Plugins/VirtualProduction/PerformanceCaptureWorkflow/Source/PerformanceCaptureWorkflow/Private/PCapSubsystem.cpp b/Engine/Plugins/VirtualProduction/PerformanceCaptureWorkflow/Source/PerformanceCaptureWorkflow/Private/PCapSubsystem.cpp
new file mode 100644
--- /dev/null
+++ b/Engine/Plugins/VirtualProduction/PerformanceCaptureWorkflow/Source/PerformanceCaptureWorkflow/Private/PCapSubsystem.cpp
@@ -0,0 +1,130 @@
+// Copyright Epic Games, Inc. All Rights Reserved.
+
+
+#include "PCapSubsystem.h"
+#include "PCapDatabase.h"
+#include "PCapSettings.h"
+#include "AssetRegistry/AssetRegistryModule.h"
+#include "Editor/TransBuffer.h"
+#include "Types/MVVMViewModelCollection.h"
+#include "Misc/CoreDelegates.h"
+#include "UObject/Package.h"
+#include "Editor.h"
+
+
+class FAssetRegistryModule;
+
+void UPerformanceCaptureSubsystem::Initialize(FSubsystemCollectionBase& Collection)
+{
+	Super::Initialize(Collection);
+	
+	ViewModelCollection = NewObject<UMVVMViewModelCollectionObject>(this);
+
+	EngineInitCompleteDelegate = FCoreDelegates::OnFEngineLoopInitComplete.AddUObject(this,&UPerformanceCaptureSubsystem::OnEngineInitComplete);
+}
+
+void UPerformanceCaptureSubsystem::OnEngineInitComplete()
+
+{
+	EngineInitCompleteDelegate.Reset();
+
+	InitiateDatabaseHelper();
+	InitiateViewModelCollection();
+
+	const UPerformanceCaptureSettings* Settings = UPerformanceCaptureSettings::GetPerformanceCaptureSettings();
+
+	/** Force load the stage class*/
+	UClass* StageRootClass =  Settings->StageRoot.LoadSynchronous();
+
+	/** Force load the UI class - this is needed so the UI class is available when the user's layout is created*/
+	UEditorUtilityWidgetBlueprint* MocapManagerUI =  Settings->MocapManagerUI.LoadSynchronous();
+
+	/** Get the Asset Registry and bind to the Remove, Duplicate and Add asset events*/
+	AssetRegistry = &FModuleManager::LoadModuleChecked<FAssetRegistryModule>(AssetRegistryConstants::ModuleName).Get();
+
+	AssetRegistry->OnAssetRemoved().AddUObject(this, &UPerformanceCaptureSubsystem::OnAssetRemoved);
+	AssetRegistry->OnAssetRenamed().AddUObject(this, &UPerformanceCaptureSubsystem::OnAssetRenamed);
+	AssetRegistry->OnAssetAdded().AddUObject(this, &UPerformanceCaptureSubsystem::OnAssetAdded);
+	FCoreUObjectDelegates::OnObjectPropertyChanged.AddUObject(this, &UPerformanceCaptureSubsystem::OnObjectPropertyChanged);
+
+#if WITH_EDITOR
+	if (GEditor)
+		{
+			if (UTransBuffer* TransactionBuffer = Cast<UTransBuffer>(GEditor->Trans))
+			{
+				TransactionBuffer->OnUndo().AddUObject(this, &UPerformanceCaptureSubsystem::OnEditorUndo);
+				TransactionBuffer->OnRedo().AddUObject(this, &UPerformanceCaptureSubsystem::OnEditorRedo);
+			}
+	}
+#endif
+	
+
+}
+
+void UPerformanceCaptureSubsystem::InitiateDatabaseHelper()
+{
+	/**Instantiate the database helper class*/
+	const UPerformanceCaptureSettings* Settings = UPerformanceCaptureSettings::GetPerformanceCaptureSettings();
+	if(UClass* HelperClass =  Settings->DatabaseHelperClass.LoadSynchronous())
+	{
+		DatabaseHelper = NewObject<UPerformanceCaptureDatabaseHelper>(GetTransientPackage(), HelperClass);
+	}
+}
+
+void UPerformanceCaptureSubsystem::InitiateViewModelCollection()
+{
+	/**Instantiate the viewmodel class*/
+	const UPerformanceCaptureSettings* Settings = UPerformanceCaptureSettings::GetPerformanceCaptureSettings();
+	if(UClass* ViewModelClass = Settings->ViewModelClass.LoadSynchronous())
+	{
+		PerformanceCaptureViewModel = NewObject<UMVVMViewModelBase>(GetTransientPackage(), ViewModelClass);
+		FMVVMViewModelContext Context;
+		Context.ContextClass = ViewModelClass;
+		Context.ContextName = "PerformanceCaptureWorkflow";
+		ViewModelCollection->AddViewModelInstance(Context, PerformanceCaptureViewModel);
+	}
+}
+
+void UPerformanceCaptureSubsystem::OnAssetRemoved(const FAssetData& InAssetData) const
+{
+	OnPCapAssetRemoved.Broadcast(InAssetData);
+}
+
+void UPerformanceCaptureSubsystem::OnAssetRenamed(const FAssetData& InAssetData, const FString& OldName) const
+{
+	OnPCapAssetRenamed.Broadcast(InAssetData, OldName);
+}
+
+void UPerformanceCaptureSubsystem::OnAssetAdded(const FAssetData& InAssetData) const
+{
+	OnPCapAssetAdded.Broadcast(InAssetData);
+}
+
+void UPerformanceCaptureSubsystem::OnObjectPropertyChanged(UObject* Asset, FPropertyChangedEvent& PropertyChangedEvent)
+{
+	AActor* Actor = Cast<AActor>(Asset);
+	if(Actor)
+	{
+		OnPCapActorModified.Broadcast(Actor);
+	}
+}
+
+void UPerformanceCaptureSubsystem::OnEditorUndo(const FTransactionContext& TransactionContext, bool bSucceeded) const
+{
+	if(bSucceeded)
+	{
+		OnPCapEditorUndo.Broadcast(true);
+	}
+}
+
+void UPerformanceCaptureSubsystem::OnEditorRedo(const FTransactionContext& TransactionContext, bool bSucceeded) const
+{
+	if(bSucceeded)
+	{
+		OnPCapEditorRedo.Broadcast(true);
+	}
+}
+
+
+
+
diff --git a/Engine/Plugins/VirtualProduction/PerformanceCaptureWorkflow/Source/PerformanceCaptureWorkflow/Private/PCapSubsystem.h b/Engine/Plugins/VirtualProduction/PerformanceCaptureWorkflow/Source/PerformanceCaptureWorkflow/Private/PCapSubsystem.h
new file mode 100644
--- /dev/null
+++ b/Engine/Plugins/VirtualProduction/PerformanceCaptureWorkflow/Source/PerformanceCaptureWorkflow/Private/PCapSubsystem.h
@@ -0,0 +1,118 @@
+// Copyright Epic Games, Inc. All Rights Reserved.
+
+#pragma once
+
+#include "CoreMinimal.h"
+#include "Subsystems/EngineSubsystem.h"
+#include "MVVMViewModelBase.h"
+#include "AssetRegistry/AssetRegistryModule.h"
+#include "AssetRegistry/AssetData.h"
+#include "UObject/Object.h"
+
+
+#include "PCapSubsystem.generated.h"
+
+class UPerformanceCaptureDatabaseHelper;
+class UMVVMViewModelCollectionObject;
+
+#if WITH_EDITOR
+struct FTransactionContext;
+#endif
+
+/**
+ * 
+ */
+UCLASS(Category="Performance Capture")
+class PERFORMANCECAPTUREWORKFLOW_API UPerformanceCaptureSubsystem : public UEngineSubsystem
+{
+	GENERATED_BODY()
+
+public:
+	// ~Begin UEditorSubsystem
+	virtual void Initialize(FSubsystemCollectionBase& Collection) override;
+	// ~End UEditorSubsystem
+
+	UFUNCTION(BlueprintCallable, Category="Performance Capture|Database")
+	UPerformanceCaptureDatabaseHelper* GetDatabaseHelper() const
+	{
+		return DatabaseHelper;
+	}
+
+	UFUNCTION(BlueprintCallable, Category="Performance Capture|Viewmodel")
+	UMVVMViewModelCollectionObject* GetViewModelCollection() const
+	{
+		return ViewModelCollection;
+	}
+
+#if WITH_EDITORONLY_DATA
+	
+	DECLARE_DYNAMIC_MULTICAST_DELEGATE_OneParam(FOnPCapAssetRemoved, FAssetData, DeletedAsset);
+	/** Editor Only*/
+	UPROPERTY(BlueprintAssignable, Category="Performance Capture|Assets")
+	FOnPCapAssetRemoved OnPCapAssetRemoved;
+
+	DECLARE_DYNAMIC_MULTICAST_DELEGATE_TwoParams(FOnPCapAssetRenamed, FAssetData, RenamedAsset, FString, OldName);
+	/** Editor Only*/
+	UPROPERTY(BlueprintAssignable, Category="Performance Capture|Assets")
+	FOnPCapAssetRenamed OnPCapAssetRenamed;
+
+	DECLARE_DYNAMIC_MULTICAST_DELEGATE_OneParam(FOnPCapAssetAdded, FAssetData, NewAsset);
+	/** Editor Only*/
+	UPROPERTY(BlueprintAssignable, Category="Performance Capture|Assets")
+	FOnPCapAssetAdded OnPCapAssetAdded;
+
+	DECLARE_DYNAMIC_MULTICAST_DELEGATE_OneParam(FOnPCapActorModified, AActor*, Actor);
+	/** Editor Only*/
+	UPROPERTY(BlueprintAssignable, Category="Performance Capture|Workflow")
+	FOnPCapActorModified OnPCapActorModified;
+
+	DECLARE_DYNAMIC_MULTICAST_DELEGATE_OneParam(FOnPCapEditorUndo, bool, bSuccess);
+	/** Editor Only*/
+	UPROPERTY(BlueprintAssignable, Category="Performance Capture|Workflow")
+	FOnPCapEditorUndo OnPCapEditorUndo;
+
+	DECLARE_DYNAMIC_MULTICAST_DELEGATE_OneParam(FOnPCapEditorRedo, bool, bSuccess);
+	/** Editor Only*/
+	UPROPERTY(BlueprintAssignable, Category="Performance Capture|Workflow")
+	FOnPCapEditorRedo OnPCapEditorRedo;
+
+
+#endif
+
+protected:
+	void OnEngineInitComplete();
+
+	void InitiateDatabaseHelper();
+
+	void InitiateViewModelCollection();
+	
+	UPROPERTY(Transient)
+	TObjectPtr<UPerformanceCaptureDatabaseHelper> DatabaseHelper;
+
+	UPROPERTY(Transient)
+	TObjectPtr<UMVVMViewModelBase> PerformanceCaptureViewModel;
+	
+	FDelegateHandle EngineInitCompleteDelegate;
+
+	IAssetRegistry* AssetRegistry;
+
+private:
+	UPROPERTY(Transient)
+	TObjectPtr<UMVVMViewModelCollectionObject> ViewModelCollection;
+
+	void OnAssetRemoved(const FAssetData& InAssetData) const;
+
+	void OnAssetRenamed(const FAssetData& InAssetData, const FString& NewName) const;
+	
+	void OnAssetAdded(const FAssetData& InAssetData) const;
+
+	void OnObjectPropertyChanged(UObject* Asset, FPropertyChangedEvent& PropertyChangedEvent);
+	
+	void OnEditorUndo(const FTransactionContext& TransactionContext, bool Succeeded) const;
+
+	void OnEditorRedo(const FTransactionContext& TransactionContext, bool Succeeded) const;
+};
+
+//Todo: add post edit change to re-int the viewmodel and database helper if the user changes these once the editor has started. 
+
+
diff --git a/Engine/Plugins/VirtualProduction/PerformanceCaptureWorkflow/Source/PerformanceCaptureWorkflow/Private/PerformanceCapture.cpp b/Engine/Plugins/VirtualProduction/PerformanceCaptureWorkflow/Source/PerformanceCaptureWorkflow/Private/PerformanceCapture.cpp
new file mode 100644
--- /dev/null
+++ b/Engine/Plugins/VirtualProduction/PerformanceCaptureWorkflow/Source/PerformanceCaptureWorkflow/Private/PerformanceCapture.cpp
@@ -0,0 +1,231 @@
+// Copyright Epic Games, Inc. All Rights Reserved.
+
+#include "PerformanceCapture.h"
+
+#include "EditorUtilityWidgetBlueprint.h"
+#include "PerformanceCaptureStyle.h"
+#include "PerformanceCaptureCommands.h"
+#include "ISettingsModule.h"
+#include "LevelEditorOutlinerSettings.h"
+#include "PCapAssetDefinition.h"
+#include "PCapDatabase.h"
+#include "PCapSettings.h"
+#include "Widgets/Docking/SDockTab.h"
+#include "Widgets/Layout/SBox.h"
+#include "Widgets/Text/STextBlock.h"
+#include "ToolMenus.h"
+#include "Visualizers/LiveLinkDataPreview.h"
+
+static const FName PerformanceCapturePanelTabName("PerformanceCaptureTab");
+
+#define LOCTEXT_NAMESPACE "FPerformanceCaptureModule"
+
+DEFINE_LOG_CATEGORY(LogPCap);
+
+void FPerformanceCaptureModule::StartupModule()
+{
+	if(ISettingsModule* SettingsModule = FModuleManager::GetModulePtr<ISettingsModule>("Settings"))
+	{
+		SettingsModule->RegisterSettings("Project", "Plugins", "PerformanceCapture", LOCTEXT("RuntimeSettingsName", "Performance Capture"), LOCTEXT("RuntimeSettingsDescription", "Performance Capture"), GetMutableDefault<UPerformanceCaptureSettings>());
+	}
+
+	FPerformanceCaptureStyle::Initialize();
+	FPerformanceCaptureStyle::ReloadTextures();
+
+	FPerformanceCaptureCommands::Register();
+	
+	PluginCommands = MakeShared<FUICommandList>();
+
+	PluginCommands->MapAction(
+		FPerformanceCaptureCommands::Get().OpenPluginWindow,
+		FExecuteAction::CreateRaw(this, &FPerformanceCaptureModule::PluginButtonClicked),
+		FCanExecuteAction());
+
+	UToolMenus::RegisterStartupCallback(FSimpleMulticastDelegate::FDelegate::CreateRaw(this, &FPerformanceCaptureModule::RegisterMenus));
+	
+	FGlobalTabmanager::Get()->RegisterNomadTabSpawner(PerformanceCapturePanelTabName, FOnSpawnTab::CreateRaw(this, &FPerformanceCaptureModule::OnSpawnMocapManager))
+		.SetDisplayName(NSLOCTEXT("PerformanceCapture", "MocapManagerTabTitle", "Mocap Manager"))
+		.SetTooltipText(NSLOCTEXT("PerformanceCapture", "PerformanceCaptureTooltipText", "Open the Mocap Manager tab"))
+		.SetMenuType(ETabSpawnerMenuType::Hidden)
+		.SetIcon(FSlateIcon(FPerformanceCaptureStyle::GetStyleSetName(), "PerformanceCapture.MocapManagerTabIcon", "PerformanceCapture.MocapManagerTabIcon.Small")
+		);
+	RegisterPlacementModeItems();
+}
+
+void FPerformanceCaptureModule::ShutdownModule()
+{
+	//Clean up settings
+	if(ISettingsModule* SettingsModule = FModuleManager::GetModulePtr<ISettingsModule>("Settings"))
+	{
+		SettingsModule->UnregisterSettings("Project", "Plugins", "PerformanceCapture");
+	}
+
+	if (!IsEngineExitRequested() && GEditor && UObjectInitialized())
+	{
+		UnregisterPlacementModeItems();
+	}
+	//Clean up nomad tab spawner
+	FGlobalTabmanager::Get()->UnregisterNomadTabSpawner(PerformanceCapturePanelTabName);
+	
+
+	UToolMenus::UnRegisterStartupCallback(this);
+
+	UToolMenus::UnregisterOwner(this);
+
+	FPerformanceCaptureStyle::Shutdown();
+
+	FPerformanceCaptureCommands::Unregister();
+	
+}
+
+TSharedRef<SDockTab> FPerformanceCaptureModule::OnSpawnMocapManager(const FSpawnTabArgs& SpawnTabArgs)
+{
+	const UPerformanceCaptureSettings* Settings = GetDefault<UPerformanceCaptureSettings>();
+	
+	if(Settings->MocapManagerUI.IsValid())
+	{
+		UEditorUtilityWidgetBlueprint* MocapManagerEW = LoadObject<UEditorUtilityWidgetBlueprint>(NULL, *Settings->MocapManagerUI.ToString(), NULL, LOAD_None, NULL);
+	
+		TSharedRef<SDockTab> TabWidget = MocapManagerEW->SpawnEditorUITab(SpawnTabArgs);
+		
+			return TabWidget;
+	}
+	
+	{
+	//Define message
+	FText WidgetText = FText::Format(
+	LOCTEXT("WindowWidgetText", "Performance Capture Project settings missing a valid UI Widget"),
+	FText::FromString(TEXT("FPerformanceCaptureModule::OnSpawnMocapManager")),
+	FText::FromString(TEXT("PerformanceCapture.cpp"))
+	);
+		//Create default tab message
+		return SNew(SDockTab)
+			.TabRole(ETabRole::NomadTab)
+			[
+				// Put your tab content here!
+		
+				SNew(SBox)
+				.HAlign(HAlign_Center)
+				.VAlign(VAlign_Center)
+				[
+					SNew(STextBlock)
+					.Text(WidgetText)
+				]
+			];
+	}
+}
+
+void FPerformanceCaptureModule::PluginButtonClicked()
+{
+	FGlobalTabmanager::Get()->TryInvokeTab(PerformanceCapturePanelTabName);
+}
+
+void FPerformanceCaptureModule::RegisterMenus()
+{
+	FToolMenuOwnerScoped OwnerScoped(this);
+
+	{
+		UToolMenu* Menu = UToolMenus::Get()->ExtendMenu("LevelEditor.MainMenu.Window.VirtualProduction");
+		{
+			FToolMenuSection& Section = Menu->FindOrAddSection("VirtualProduction");
+			Section.AddMenuEntryWithCommandList(FPerformanceCaptureCommands::Get().OpenPluginWindow, PluginCommands);
+		}
+	}
+
+	{
+		UToolMenu* ToolbarMenu = UToolMenus::Get()->ExtendMenu("LevelEditor.LevelEditorToolBar");
+		{
+			FToolMenuSection& Section = ToolbarMenu->FindOrAddSection("Settings");
+			{
+				FToolMenuEntry& Entry = Section.AddEntry(FToolMenuEntry::InitToolBarButton(FPerformanceCaptureCommands::Get().OpenPluginWindow));
+				Entry.SetCommandList(PluginCommands);
+			}
+		}
+	}
+}
+
+void FPerformanceCaptureModule::UnregisterPlacementModeItems()
+{
+	IPlacementModeModule& PlacementModeModule = IPlacementModeModule::Get();
+
+	for (TOptional<FPlacementModeID>& PlaceActor : PlaceActors)
+	{
+		if (PlaceActor.IsSet())
+		{
+			PlacementModeModule.UnregisterPlaceableItem(*PlaceActor);
+		}
+	}
+
+	PlaceActors.Empty();
+}
+
+const FPlacementCategoryInfo* FPerformanceCaptureModule::GetVirtualProductionCategoryRegisteredInfo() const
+{
+	IPlacementModeModule& PlacementModeModule = IPlacementModeModule::Get();
+
+	if (const FPlacementCategoryInfo* RegisteredInfo = PlacementModeModule.GetRegisteredPlacementCategory(FLevelEditorOutlinerBuiltInCategories::VirtualProduction()))
+	{
+		return RegisteredInfo;
+	}
+	else
+	{
+		FPlacementCategoryInfo Info(
+			LOCTEXT("VirtualProductionCategoryName", "Virtual Production"),
+			FSlateIcon(FPerformanceCaptureStyle::Get().GetStyleSetName(), "PlacementBrowser.Icons.VirtualProduction"),
+			FLevelEditorOutlinerBuiltInCategories::VirtualProduction(),
+			TEXT("PMVirtualProduction"),
+			25 // Determines where the category shows up in the list with respect to the others.
+		);
+		Info.ShortDisplayName = LOCTEXT("VirtualProductionShortCategoryName", "VP");
+		IPlacementModeModule::Get().RegisterPlacementCategory(Info);
+
+		return PlacementModeModule.GetRegisteredPlacementCategory(FLevelEditorOutlinerBuiltInCategories::VirtualProduction());
+	}
+}
+
+void FPerformanceCaptureModule::RegisterPlacementModeItems()
+{
+	auto RegisterPlaceActors = [&]() -> void
+	{
+		if (!GEditor)
+		{
+			return;
+		}
+
+		const FPlacementCategoryInfo* Info = GetVirtualProductionCategoryRegisteredInfo();
+
+		if (!Info)
+		{
+			UE_LOG(LogPCap, Warning, TEXT("Could not find or create VirtualProduction Place Actor Category"));
+			return;
+		}
+
+		/*Register the Live Link Data Prevew Actor*/
+		PlaceActors.Add(IPlacementModeModule::Get().RegisterPlaceableItem(Info->UniqueHandle, MakeShared<FPlaceableItem>(
+			*ALiveLinkDataPreview::StaticClass(),
+			FAssetData(ALiveLinkDataPreview::StaticClass()),
+			NAME_None,
+			NAME_None,
+			TOptional<FLinearColor>(),
+			TOptional<int32>(),
+			NSLOCTEXT("PlacementMode", "LiveLinkDataPreview", "Live Link Data Preview")
+		)));
+	};
+
+	if (FApp::CanEverRender())
+	{
+		if (GEngine && GEngine->IsInitialized())
+		{
+			RegisterPlaceActors();
+		}
+		else
+		{
+			PostEngineInitHandle = FCoreDelegates::OnPostEngineInit.AddLambda(RegisterPlaceActors);
+		}
+	}
+}
+
+#undef LOCTEXT_NAMESPACE
+	
+IMPLEMENT_MODULE(FPerformanceCaptureModule, PerformanceCaptureWorkflow)
+
diff --git a/Engine/Plugins/VirtualProduction/PerformanceCaptureWorkflow/Source/PerformanceCaptureWorkflow/Private/PerformanceCaptureCommands.cpp b/Engine/Plugins/VirtualProduction/PerformanceCaptureWorkflow/Source/PerformanceCaptureWorkflow/Private/PerformanceCaptureCommands.cpp
new file mode 100644
--- /dev/null
+++ b/Engine/Plugins/VirtualProduction/PerformanceCaptureWorkflow/Source/PerformanceCaptureWorkflow/Private/PerformanceCaptureCommands.cpp
@@ -0,0 +1,12 @@
+// Copyright Epic Games, Inc. All Rights Reserved.
+
+#include "PerformanceCaptureCommands.h"
+
+#define LOCTEXT_NAMESPACE "FPerformanceCaptureModule"
+
+void FPerformanceCaptureCommands::RegisterCommands()
+{
+	UI_COMMAND(OpenPluginWindow, "Mocap Manager", "Open Mocap Manager", EUserInterfaceActionType::Button, FInputChord());
+}
+
+#undef LOCTEXT_NAMESPACE
\ No newline at end of file
diff --git a/Engine/Plugins/VirtualProduction/PerformanceCaptureWorkflow/Source/PerformanceCaptureWorkflow/Private/PerformanceCaptureStyle.cpp b/Engine/Plugins/VirtualProduction/PerformanceCaptureWorkflow/Source/PerformanceCaptureWorkflow/Private/PerformanceCaptureStyle.cpp
new file mode 100644
--- /dev/null
+++ b/Engine/Plugins/VirtualProduction/PerformanceCaptureWorkflow/Source/PerformanceCaptureWorkflow/Private/PerformanceCaptureStyle.cpp
@@ -0,0 +1,91 @@
+// Copyright Epic Games, Inc. All Rights Reserved.
+
+#include "PerformanceCaptureStyle.h"
+#include "Styling/SlateStyleRegistry.h"
+#include "Framework/Application/SlateApplication.h"
+#include "Interfaces/IPluginManager.h"
+#include "Styling/SlateStyleMacros.h"
+#include "Styling/SlateStyle.h"
+
+#define RootToContentDir Style->RootToContentDir
+
+TSharedPtr<FSlateStyleSet> FPerformanceCaptureStyle::StyleInstance = nullptr;
+FColor FPerformanceCaptureStyle::TypeColor(104,49,178);
+
+
+void FPerformanceCaptureStyle::Initialize()
+{
+	if(!StyleInstance.IsValid())
+	{
+		StyleInstance = Create();
+		FSlateStyleRegistry::RegisterSlateStyle(*StyleInstance);
+	}
+}
+
+void FPerformanceCaptureStyle::Shutdown()
+{
+	FSlateStyleRegistry::UnRegisterSlateStyle(*StyleInstance);
+	ensure(StyleInstance.IsUnique());
+	StyleInstance.Reset();
+}
+
+FName FPerformanceCaptureStyle::GetStyleSetName()
+{
+	static FName StyleSetName(TEXT("PerformanceCaptureStyle"));
+	return StyleSetName;
+}
+
+const FVector2D Icon16x16(16.0f, 16.0f);
+const FVector2D Icon20x20(20.0f, 20.0f);
+const FVector2D Icon64x64(64.0f, 64.0f);
+
+TSharedRef< FSlateStyleSet > FPerformanceCaptureStyle::Create()
+{
+	TSharedRef< FSlateStyleSet > Style = MakeShareable(new FSlateStyleSet("PerformanceCaptureStyle"));
+	Style->SetContentRoot(IPluginManager::Get().FindPlugin("PerformanceCaptureWorkflow")->GetBaseDir() / TEXT("Resources"));
+
+	//TODO: note to future self: Styles must be ClassThumbnail.Name of asset!!! */
+
+	Style->Set("PerformanceCapture.OpenPluginWindow", new IMAGE_BRUSH_SVG(TEXT("ButtonIcon"), Icon20x20));
+
+	Style->Set("PerformanceCapture.MocapManagerTabIcon", new IMAGE_BRUSH_SVG(TEXT("ButtonIcon"), Icon20x20));
+
+	Style->Set("PerformanceCapture.MocapManagerTabIcon.Small", new IMAGE_BRUSH_SVG(TEXT("ButtonIcon"), Icon16x16));
+	
+	Style->Set("PlacementBrowser.Icons.VirtualProduction", new IMAGE_BRUSH_SVG(TEXT("VirtualProduction"), Icon16x16));
+
+	Style->Set("ClassThumbnail.PCapDataTable", new IMAGE_BRUSH_SVG(TEXT("DataTableThumbnail"), Icon64x64));
+
+	Style->Set("ClassIcon.PCapDataTable", new IMAGE_BRUSH_SVG(TEXT("DataTableIcon"), Icon16x16));
+
+	Style->Set("ClassThumbnail.PCapPerformerDataAsset", new IMAGE_BRUSH_SVG(TEXT("PerformerDataAsset"), Icon64x64));
+
+	Style->Set("ClassIcon.PCapPerformerDataAsset", new IMAGE_BRUSH_SVG(TEXT("PerformerDataAsset"), Icon16x16));
+
+	Style->Set("ClassThumbnail.PCapCharacterDataAsset", new IMAGE_BRUSH_SVG(TEXT("CharacterDataAsset"), Icon64x64));
+	
+	Style->Set("ClassIcon.PCapCharacterDataAsset", new IMAGE_BRUSH_SVG(TEXT("CharacterDataAsset"), Icon16x16));
+
+	Style->Set("ClassThumbnail.PCapPropDataAsset", new IMAGE_BRUSH_SVG(TEXT("PropDataAsset"), Icon64x64));
+
+	Style->Set("ClassIcon.PCapPropDataAsset", new IMAGE_BRUSH_SVG(TEXT("PropDataAsset"), Icon16x16));
+	
+	Style->Set("ClassThumbnail.PCapSessionTemplate", new IMAGE_BRUSH_SVG(TEXT("PcapSessionTemplate_64"), Icon64x64));
+
+	Style->Set("ClassIcon.PCapSessionTemplate", new IMAGE_BRUSH_SVG(TEXT("PcapSessionTemplate_16"), Icon16x16));
+	
+	return Style;
+}
+
+void FPerformanceCaptureStyle::ReloadTextures()
+{
+	if (FSlateApplication::IsInitialized())
+	{
+		FSlateApplication::Get().GetRenderer()->ReloadTextureResources();
+	}
+}
+
+const ISlateStyle& FPerformanceCaptureStyle::Get()
+{
+	return *StyleInstance;
+}
diff --git a/Engine/Plugins/VirtualProduction/PerformanceCaptureWorkflow/Source/PerformanceCaptureWorkflow/Private/Visualizers/LiveLinkDataPreview.cpp b/Engine/Plugins/VirtualProduction/PerformanceCaptureWorkflow/Source/PerformanceCaptureWorkflow/Private/Visualizers/LiveLinkDataPreview.cpp
new file mode 100644
--- /dev/null
+++ b/Engine/Plugins/VirtualProduction/PerformanceCaptureWorkflow/Source/PerformanceCaptureWorkflow/Private/Visualizers/LiveLinkDataPreview.cpp
@@ -0,0 +1,203 @@
+// Copyright Epic Games, Inc. All Rights Reserved.
+
+
+#include "LiveLinkDataPreview.h"
+
+#include "Roles/LiveLinkAnimationRole.h"
+#include "Roles/LiveLinkCameraRole.h"
+#include "Roles/LiveLinkLocatorRole.h"
+#include "Roles/LiveLinkTransformRole.h"
+#include "ILiveLinkClient.h"
+#include "Features/IModularFeatures.h"
+#include "LiveLinkDataPreviewComponent.h"
+#include "Components/BillboardComponent.h"
+#include "Materials/MaterialInstanceDynamic.h"
+#include "Engine/StaticMesh.h"
+#include "Engine/Texture2D.h"
+#include "Engine/World.h"
+
+
+// Sets default values
+ALiveLinkDataPreview::ALiveLinkDataPreview()
+{
+	// Set this actor to call Tick() every frame.  You can turn this off to improve performance if you don't need it.
+	PrimaryActorTick.bCanEverTick = true;
+	bEvaluateLiveLink = true;
+
+	BoneMesh = FSoftObjectPath("/PerformanceCaptureWorkflow/Visualizers/SM_Bone.SM_Bone");
+	JointMesh = FSoftObjectPath("/PerformanceCaptureWorkflow/Visualizers/SM_Joint.SM_Joint");
+	AxisMesh = FSoftObjectPath("/PerformanceCaptureWorkflow/Visualizers/SM_JointAxis.SM_JointAxis");
+	TransformMesh = FSoftObjectPath("/PerformanceCaptureWorkflow/Visualizers/SM_TransformAxis.SM_TransformAxis");
+	LocatorMesh = FSoftObjectPath("/PerformanceCaptureWorkflow/Visualizers/SM_MarkerSphere.SM_MarkerSphere");
+	CameraMesh = FSoftObjectPath("/PerformanceCaptureWorkflow/Visualizers/SM_LiveLinkCamera.SM_LiveLinkCamera");
+	SpriteTexture = FSoftObjectPath("/PerformanceCaptureWorkflow/Common/Icons/Location_256.Location_256");
+
+	USceneComponent* SceneRoot = CreateDefaultSubobject<USceneComponent>("Root");
+	SetRootComponent(SceneRoot);
+	BillboardComponent = CreateDefaultSubobject<UBillboardComponent>("Sprite");
+	BillboardComponent->SetupAttachment(SceneRoot);
+	UTexture2D* SpriteTex = SpriteTexture.LoadSynchronous();
+	if(SpriteTex)
+	{
+		BillboardComponent->SetSprite(SpriteTex);
+		BillboardComponent->bIsScreenSizeScaled = true;
+		BillboardComponent->ScreenSize = 0.0006;
+		BillboardComponent->SetRelativeLocation(FVector(0,0,10));
+	}
+}
+
+// Called every frame
+void ALiveLinkDataPreview::Tick(float DeltaTime)
+{
+	Super::Tick(DeltaTime);
+
+	IModularFeatures& ModularFeatures = IModularFeatures::Get();
+	if (ModularFeatures.IsModularFeatureAvailable(ILiveLinkClient::ModularFeatureName))
+	{
+		ILiveLinkClient& LiveLinkClient = IModularFeatures::Get().GetModularFeature<ILiveLinkClient>(ILiveLinkClient::ModularFeatureName);
+	
+		TArray<FLiveLinkSubjectKey> SubjectKeys = LiveLinkClient.GetSubjects(false, false);
+
+		if(SubjectKeys.Num() != CachedSubjects.Num())
+		{
+			InitializeSubjects();
+		}
+		CachedSubjects = SubjectKeys;
+	}
+}
+
+void ALiveLinkDataPreview::OnConstruction(const FTransform& Transform)
+{
+	Super::OnConstruction(Transform);
+
+	InitializeSubjects();
+}
+
+void ALiveLinkDataPreview::PostEditChangeProperty(FPropertyChangedEvent& PropertyChangedEvent)
+{
+	Super::PostEditChangeProperty(PropertyChangedEvent);
+
+	if(PropertyChangedEvent.Property != nullptr)
+	{
+		if(PropertyChangedEvent.Property->GetFName() == GET_MEMBER_NAME_CHECKED(ALiveLinkDataPreview, bEvaluateLiveLink))
+		{
+			TArray<ULiveLinkDataPreviewComponent*> DebugComponents;
+			GetComponents(DebugComponents);
+
+			for (ULiveLinkDataPreviewComponent* DebugComponent : DebugComponents)
+			{
+				DebugComponent->SetEvaluateLiveLinkData(bEvaluateLiveLink);
+			}
+		}
+	}
+}
+
+void ALiveLinkDataPreview::SetEnableLiveLinkData(bool bNewEvaluate)
+{
+	bEvaluateLiveLink = bNewEvaluate;
+
+	TArray<ULiveLinkDataPreviewComponent*> DebugComponents;
+	GetComponents(DebugComponents);
+
+	for (ULiveLinkDataPreviewComponent* DebugComponent : DebugComponents)
+	{
+		DebugComponent->SetEvaluateLiveLinkData(bEvaluateLiveLink);
+	}
+}
+
+void ALiveLinkDataPreview::InitializeSubjects()
+{
+	if((GetWorld() && !GetWorld()->IsPreviewWorld())) //Do not spawn the preview components unless we are a non-preview world
+	{
+		IModularFeatures& ModularFeatures = IModularFeatures::Get();
+		if (ModularFeatures.IsModularFeatureAvailable(ILiveLinkClient::ModularFeatureName))
+		{
+			ILiveLinkClient& LiveLinkClient = IModularFeatures::Get().GetModularFeature<ILiveLinkClient>(ILiveLinkClient::ModularFeatureName);
+
+			TArray<FLiveLinkSubjectKey> SubjectKeys = LiveLinkClient.GetSubjects(false, false);
+
+			Subjects.Empty();
+
+			for (FLiveLinkSubjectKey Subjectkey : SubjectKeys)
+			{
+				FLiveLinkSubjectName Subject = Subjectkey.SubjectName;
+				Subjects.Emplace(Subject);
+			}
+
+			TArray<ULiveLinkDataPreviewComponent*> DataPreviewComponents;
+			GetComponents(DataPreviewComponents);
+
+			for (ULiveLinkDataPreviewComponent* DataPreviewComponent : DataPreviewComponents)
+			{
+				DataPreviewComponent->DestroyComponent();
+			}
+			
+			for (FLiveLinkSubjectName Subject : Subjects)
+			{
+				if(LiveLinkClient.IsSubjectEnabled(Subject))
+				{
+					ULiveLinkDataPreviewComponent* DataPreviewComponent;
+					DataPreviewComponent = Cast<ULiveLinkDataPreviewComponent>(AddComponentByClass(ULiveLinkDataPreviewComponent::StaticClass(), false, FTransform::Identity, false));
+					
+					DataPreviewComponent->SubjectName = Subject;
+					DataPreviewComponent->bEvaluateLiveLink = bEvaluateLiveLink;
+					DataPreviewComponent->bDrawLabels = bDrawLabels;
+
+					if(LiveLinkClient.DoesSubjectSupportsRole_AnyThread(Subject, ULiveLinkAnimationRole::StaticClass()))
+					{
+						if(UStaticMesh* StaticMesh = BoneMesh.LoadSynchronous())
+						{
+							DataPreviewComponent->SetStaticMesh(StaticMesh);
+							SetMaterialInstance(DataPreviewComponent);
+						}
+						DataPreviewComponent->BoneVisualType = ELiveLinkVisualBoneType::Bone;
+					}
+					else if(LiveLinkClient.DoesSubjectSupportsRole_AnyThread(Subject, ULiveLinkLocatorRole::StaticClass()))
+					{
+						if(UStaticMesh* StaticMesh = LocatorMesh.LoadSynchronous())
+						{
+							DataPreviewComponent->SetStaticMesh(StaticMesh);
+							SetMaterialInstance(DataPreviewComponent);
+						}
+					}
+					else if(LiveLinkClient.DoesSubjectSupportsRole_AnyThread(Subject, ULiveLinkCameraRole::StaticClass()))
+					{
+						
+						if(UStaticMesh* StaticMesh = CameraMesh.LoadSynchronous())
+						{
+							DataPreviewComponent->SetStaticMesh(StaticMesh);
+							SetMaterialInstance(DataPreviewComponent);
+						}
+					}
+					else if(LiveLinkClient.DoesSubjectSupportsRole_AnyThread(Subject, ULiveLinkTransformRole::StaticClass()))
+					{
+						;
+						if(UStaticMesh* StaticMesh = TransformMesh.LoadSynchronous())
+						{
+							DataPreviewComponent->SetStaticMesh(StaticMesh);
+						}
+					}
+				}
+			}
+		}
+	}
+}
+
+void ALiveLinkDataPreview::SetMaterialInstance(ULiveLinkDataPreviewComponent* InDataPreviewComponent)
+{
+	UMaterialInterface* DebugMaterial = InDataPreviewComponent->GetMaterial(0);
+	if(DebugMaterial)
+	{
+		UMaterialInstanceDynamic* DynamicMaterial = InDataPreviewComponent->CreateDynamicMaterialInstance(0, DebugMaterial);
+		InDataPreviewComponent->SetMaterial(0, DynamicMaterial);
+		DynamicMaterial->ClearParameterValues();
+		/*Generate a random color that is weighted toward 1 or 2 primary colors.*/
+		float Red01 = FMath::FRandRange(0.33,0.8);
+		float Green01 = FMath::FRandRange(0.33,0.8);
+		float Blue01 = FMath::FRandRange(0.33,0.8);
+		float Red02 = Red01 - Green01;
+		float Green02 = Green01 - Blue01;
+		float Blue02 = Green01 - Red01;
+		DynamicMaterial->SetVectorParameterValue(FName(TEXT("Color")), FLinearColor(FMath::RandBool() ? Red01 : Red02, FMath::RandBool() ? Green01 : Green02, FMath::RandBool() ? Blue01 : Blue02, 1.0) );
+	}
+}
diff --git a/Engine/Plugins/VirtualProduction/PerformanceCaptureWorkflow/Source/PerformanceCaptureWorkflow/Private/Visualizers/LiveLinkDataPreview.h b/Engine/Plugins/VirtualProduction/PerformanceCaptureWorkflow/Source/PerformanceCaptureWorkflow/Private/Visualizers/LiveLinkDataPreview.h
new file mode 100644
--- /dev/null
+++ b/Engine/Plugins/VirtualProduction/PerformanceCaptureWorkflow/Source/PerformanceCaptureWorkflow/Private/Visualizers/LiveLinkDataPreview.h
@@ -0,0 +1,76 @@
+// Copyright Epic Games, Inc. All Rights Reserved.
+
+#pragma once
+
+#include "CoreMinimal.h"
+#include "GameFramework/Actor.h"
+#include "LiveLinkTypes.h"
+#include "LiveLinkDataPreviewComponent.h"
+
+#include "LiveLinkDataPreview.generated.h"
+
+/**
+ * Actor for visualizing LiveLink data in the viewport/level editor
+ */
+
+UCLASS(Blueprintable, BlueprintType, DisplayName="Live Link Data Preview Actor", Category="LiveLink")
+class PERFORMANCECAPTUREWORKFLOW_API ALiveLinkDataPreview : public AActor
+{
+	GENERATED_BODY()
+
+public:
+	// Sets default values for this actor's properties
+	ALiveLinkDataPreview();
+
+	UPROPERTY(EditAnywhere, BlueprintReadWrite, Category = LiveLink, meta=(ShowOnlyInnerProperties="true"))
+	TArray<FLiveLinkSubjectName> Subjects;
+
+	TSoftObjectPtr<UStaticMesh> BoneMesh;
+	TSoftObjectPtr<UStaticMesh> JointMesh;
+	TSoftObjectPtr<UStaticMesh> AxisMesh;
+	TSoftObjectPtr<UStaticMesh> TransformMesh;
+	TSoftObjectPtr<UStaticMesh> CameraMesh;
+	TSoftObjectPtr<UStaticMesh> LocatorMesh;
+	TSoftObjectPtr<UTexture2D> SpriteTexture;
+	UPROPERTY(EditAnywhere, BlueprintReadWrite, Category = LiveLink)
+	bool bDrawLabels;
+
+	UPROPERTY(EditAnywhere, BlueprintReadWrite, Category = LiveLink)
+	bool bEvaluateLiveLink;
+
+	TArray<ULiveLinkDataPreviewComponent> Visualizers;
+	
+	UFUNCTION(CallInEditor, BlueprintCallable, DisplayName="Evaluate Live Link", Category = LiveLink)
+	void SetEnableLiveLinkData(bool bNewEvaluate);
+
+	UPROPERTY(EditAnywhere, BlueprintReadWrite, Category = Components)
+	TObjectPtr<UBillboardComponent> BillboardComponent;
+
+	protected:
+
+	UFUNCTION(BlueprintCallable, DisplayName="Initialize All Subjects", Category = LiveLink)
+	void InitializeSubjects();
+
+	void SetMaterialInstance(ULiveLinkDataPreviewComponent* InDataPreviewComponent);
+	
+	// Called every frame
+	virtual void Tick(float DeltaTime) override;
+
+	//Actor Interface
+	virtual void OnConstruction(const FTransform& Transform) override;
+	
+	//PostEnd Actor Interface
+
+#if WITH_EDITOR
+
+	virtual void PostEditChangeProperty(FPropertyChangedEvent& PropertyChangedEvent) override;
+
+	virtual bool ShouldTickIfViewportsOnly() const override
+	{
+		return true;
+	}
+#endif
+
+private:
+	TArray<FLiveLinkSubjectKey> CachedSubjects;
+};
\ No newline at end of file
diff --git a/Engine/Plugins/VirtualProduction/PerformanceCaptureWorkflow/Source/PerformanceCaptureWorkflow/Private/Visualizers/LiveLinkDataPreviewComponent.cpp b/Engine/Plugins/VirtualProduction/PerformanceCaptureWorkflow/Source/PerformanceCaptureWorkflow/Private/Visualizers/LiveLinkDataPreviewComponent.cpp
new file mode 100644
--- /dev/null
+++ b/Engine/Plugins/VirtualProduction/PerformanceCaptureWorkflow/Source/PerformanceCaptureWorkflow/Private/Visualizers/LiveLinkDataPreviewComponent.cpp
@@ -0,0 +1,371 @@
+// Copyright Epic Games, Inc. All Rights Reserved.
+
+
+#include "LiveLinkDataPreviewComponent.h"
+
+#include "AnimationCoreLibrary.h"
+#include "CanvasTypes.h"
+#include "ILiveLinkClient.h"
+#include "LevelEditorViewport.h"
+#include "Features/IModularFeatures.h"
+#include "PCapBoneVisualizer.h"
+#include "PerformanceCapture.h"
+#include "Roles/LiveLinkAnimationRole.h"
+#include "Roles/LiveLinkLocatorRole.h"
+#include "Roles/LiveLinkTransformRole.h"
+
+ULiveLinkDataPreviewComponent::ULiveLinkDataPreviewComponent(const FObjectInitializer& ObjectInitializer)
+	: bEvaluateLiveLink(true)
+	, bDrawLabels(false)
+	, bIsDirty(true)
+{
+	PrimaryComponentTick.bCanEverTick = true;
+	bTickInEditor = true; //make this component tick in editor
+	PrimaryComponentTick.TickGroup = ETickingGroup::TG_PostUpdateWork;
+}
+
+void ULiveLinkDataPreviewComponent::TickComponent(float DeltaTime, ELevelTick TickType, FActorComponentTickFunction* ThisTickFunction)
+{
+	Super::TickComponent(DeltaTime, TickType, ThisTickFunction);
+	
+	ILiveLinkClient& LiveLinkClient = IModularFeatures::Get().GetModularFeature<ILiveLinkClient>(ILiveLinkClient::ModularFeatureName);
+
+	TSubclassOf<ULiveLinkRole> SubjectRole = LiveLinkClient.GetSubjectRole_AnyThread(SubjectName);
+	
+	if (SubjectRole && !bIsDirty)
+	{
+		//Process Animation Data
+		if(LiveLinkClient.DoesSubjectSupportsRole_AnyThread(SubjectName, ULiveLinkAnimationRole::StaticClass()))
+		{
+			if(LiveLinkClient.IsSubjectValid(SubjectName)) //Check there is valid data on the subject before evaluate it. 
+			{
+				FLiveLinkSubjectFrameData SubjectFrameData;
+				LiveLinkClient.EvaluateFrame_AnyThread(SubjectName, LiveLinkClient.GetSubjectRole_AnyThread(SubjectName), SubjectFrameData);
+
+				FLiveLinkSkeletonStaticData* SkeletonStaticData = SubjectFrameData.StaticData.Cast<FLiveLinkSkeletonStaticData>();
+				FLiveLinkAnimationFrameData* AnimationFrameData = SubjectFrameData.FrameData.Cast<FLiveLinkAnimationFrameData>();
+				check(SkeletonStaticData);
+				check(AnimationFrameData);
+
+				CacheSkeletalAnimationData(SkeletonStaticData, AnimationFrameData);
+
+				if(GetInstanceCount()<1)
+				{
+					CreateInstances();
+				}
+
+				if(bEvaluateLiveLink)
+				{
+					switch (BoneVisualType)
+					{
+					case ELiveLinkVisualBoneType::Joint:
+						{
+							BatchUpdateInstancesTransforms(0, GetJointTransforms(), true, true, false);
+							break;
+						}
+					case ELiveLinkVisualBoneType::Bone:
+						{
+							BatchUpdateInstancesTransforms(0, GetBoneTransforms(), true, true, false);
+							break;
+						}
+					}
+				}
+			}
+		}
+		else if(LiveLinkClient.DoesSubjectSupportsRole_AnyThread(SubjectName, ULiveLinkLocatorRole::StaticClass()) && bEvaluateLiveLink)
+		{
+			
+			TArray<FTransform> Transforms;
+			TArray<FVector> Locators = GetLocatorPositions();
+
+			for (int32 i = 0; i < Locators.Num(); i++)
+			{
+				FVector Locator = Locators[i];
+				FTransform Transform = FTransform::Identity;
+				Transform.SetTranslation(Locator);
+				Transforms.Emplace(Transform * GetComponentTransform());
+			}
+			BatchUpdateInstancesTransforms(0, Transforms, true, true, false);
+		}
+		else if(LiveLinkClient.DoesSubjectSupportsRole_AnyThread(SubjectName, ULiveLinkTransformRole::StaticClass()) && bEvaluateLiveLink)
+		{
+			BatchUpdateInstancesTransform(0,1,GetSingleTransform() * GetComponentTransform(),true,true,false);
+
+			if (bDrawLabels)
+			{
+				FVector Location = (GetSingleTransform() * GetComponentTransform()).GetTranslation();
+				float Time = 1.0;
+				UWorld* World = GetWorld();
+				DrawDebugString(World, Location, TEXT("this is some text"), GetOwner(), FColor::Red, Time, false, 1.0);
+			}
+		}
+	}
+	if (bIsDirty)
+	{
+		CreateInstances();
+	}
+}
+
+void ULiveLinkDataPreviewComponent::PostEditChangeProperty(FPropertyChangedEvent& PropertyChangedEvent)
+{
+	Super::PostEditChangeProperty(PropertyChangedEvent);
+
+	if(PropertyChangedEvent.Property != nullptr)
+	{
+		if(PropertyChangedEvent.Property->GetFName() == GET_MEMBER_NAME_CHECKED(ULiveLinkDataPreviewComponent, SubjectName))
+		{
+			ClearInstances();
+			bIsDirty = true;
+		}
+		if(PropertyChangedEvent.Property->GetFName() == GET_MEMBER_NAME_CHECKED(ULiveLinkDataPreviewComponent, bEvaluateLiveLink))
+		{
+			ClearInstances();
+			bIsDirty = true;
+		}
+		if(PropertyChangedEvent.Property->GetFName() == GET_MEMBER_NAME_CHECKED(ULiveLinkDataPreviewComponent, BoneVisualType))
+		{
+			ClearInstances();
+			bIsDirty = true;
+		}
+		if(PropertyChangedEvent.Property->GetFName() == GET_MEMBER_NAME_CHECKED(ULiveLinkDataPreviewComponent, bDrawLabels))
+		{
+			ClearInstances();
+			bIsDirty = true;
+		}
+	}
+	
+}
+
+void ULiveLinkDataPreviewComponent::OnRegister()
+{
+	Super::OnRegister();
+	
+	ClearInstances();
+
+	bIsDirty = true;
+}
+
+void ULiveLinkDataPreviewComponent::OnUnregister()
+{
+	Super::OnUnregister();
+
+	ClearInstances();
+	
+}
+
+void ULiveLinkDataPreviewComponent::GetTransformRootSpace(const int32 InTransformIndex, FTransform& OutTransform) const
+{
+	// Case: Root joint or invalid
+	OutTransform = FTransform::Identity;
+	
+	if (IsValidTransformIndex(InTransformIndex, CachedAnimationData))
+	{
+		TPair<bool, FTransform>& RootSpaceCache = CachedRootSpaceTransforms[InTransformIndex];
+		bool bHasValidCache = RootSpaceCache.Key;
+		// Case: Have Cached Value
+		if (bHasValidCache)
+		{
+			OutTransform = RootSpaceCache.Value;
+		}
+		// Case: Need to generate Cache
+		else
+		{
+			const TArray<int32>& BoneParents = CachedSkeletonData.GetBoneParents();
+			int32 ParentIndex = BoneParents[InTransformIndex];
+
+			const FTransform& LocalSpaceTransform = CachedAnimationData.Transforms[InTransformIndex];
+
+			FTransform ParentRootSpaceTransform;
+
+			if(InTransformIndex == 0 && CachedSkeletonData.BoneParents[0] == 0) //special case handling for virtual subjects where bone[0] has parent bone[0]
+			{
+				ParentRootSpaceTransform = FTransform::Identity;
+			}
+			else
+			{
+				GetTransformRootSpace(ParentIndex, ParentRootSpaceTransform);
+			}
+
+			OutTransform = LocalSpaceTransform * ParentRootSpaceTransform;
+
+			// Save cached results
+			RootSpaceCache.Key = true;
+			RootSpaceCache.Value = OutTransform;
+		}
+	}
+}
+
+void ULiveLinkDataPreviewComponent::CreateInstances()
+{
+	ILiveLinkClient& LiveLinkClient = IModularFeatures::Get().GetModularFeature<ILiveLinkClient>(ILiveLinkClient::ModularFeatureName);
+
+	TSubclassOf<ULiveLinkRole> SubjectRole = LiveLinkClient.GetSubjectRole_AnyThread(SubjectName);
+	if (SubjectRole)
+	{
+		//Process Animation Data
+		if(LiveLinkClient.DoesSubjectSupportsRole_AnyThread(SubjectName, ULiveLinkAnimationRole::StaticClass()))
+		{
+			switch (BoneVisualType)
+			{
+			case ELiveLinkVisualBoneType::Joint:
+				{
+					AddInstances(GetJointTransforms(), false, true, false);
+					break;
+				}
+			case ELiveLinkVisualBoneType::Bone:
+				{
+					AddInstances(GetBoneTransforms(), false, true, false);
+					break;
+				}
+			}
+			UE_LOG(LogPCap, Display, TEXT("Created Instances %d"), GetInstanceCount());
+			UE_LOG(LogPCap, Display, TEXT("Created Instances %s"), bIsDirty ? TEXT("true") : TEXT("false") );
+			bIsDirty=false;
+		}
+		else if(LiveLinkClient.DoesSubjectSupportsRole_AnyThread(SubjectName, ULiveLinkLocatorRole::StaticClass()))
+		{
+			TArray<FTransform> Transforms;
+			TArray<FVector> Locators = GetLocatorPositions();
+
+			for (int32 i = 0; i < Locators.Num(); i++)
+			{
+				FVector Locator = Locators[i];
+				FTransform Transform = FTransform::Identity;
+				Transform.SetTranslation(Locator);
+				Transforms.Emplace(Transform * GetComponentTransform());
+			}
+			AddInstances(Transforms, false, true, false);
+			bIsDirty=false;
+		}
+		else if(LiveLinkClient.DoesSubjectSupportsRole_AnyThread(SubjectName, ULiveLinkTransformRole::StaticClass()))
+		{
+			AddInstance(GetSingleTransform() * GetComponentTransform(),true);
+			bIsDirty=false;
+		}
+	}
+}
+
+TArray<FTransform> ULiveLinkDataPreviewComponent::GetBoneTransforms() const
+{
+	TArray<FTransform> JointTransforms;
+		
+		FTransform BoneTransform;
+		FTransform ParentBoneTransform;
+		double BoneLength;
+		FVector Aim = FVector( 0, 0, 1);
+
+		const TArray<FTransform> BoneTransforms = CachedAnimationData.Transforms;
+		const TArray<int32> BoneParents = CachedSkeletonData.BoneParents;
+
+					
+		for (int32 i=0; i < CachedSkeletonData.BoneNames.Num(); i++)
+		{
+			GetTransformRootSpace(i, BoneTransform);
+			BoneTransform = BoneTransform * GetComponentTransform();
+			GetTransformRootSpace(CachedSkeletonData.BoneParents[i],ParentBoneTransform);
+			ParentBoneTransform = ParentBoneTransform * GetComponentTransform();
+			
+			BoneLength = FVector::Distance(BoneTransform.GetTranslation(), ParentBoneTransform.GetTranslation());
+			BoneTransform.SetScale3D(FVector(1, 1, BoneLength));
+			FQuat DiffRotation = AnimationCore::SolveAim(BoneTransform, ParentBoneTransform.GetLocation(), Aim.GetSafeNormal(), false, FVector(1, 1, 1), float(0));
+			BoneTransform.SetRotation(DiffRotation);
+			JointTransforms.Add(BoneTransform);
+		}
+	return JointTransforms;
+}
+
+TArray<FTransform> ULiveLinkDataPreviewComponent::GetJointTransforms() const
+{
+	TArray<FTransform> BoneTransforms;
+	
+				
+		int32 i = 0;
+		FTransform BoneTransform;
+
+		for (i=0 ; i < CachedSkeletonData.BoneNames.Num() ; i++)
+		{
+			GetTransformRootSpace(i, BoneTransform);
+			BoneTransform = BoneTransform * GetComponentTransform();
+			BoneTransforms.Add(BoneTransform);
+		}
+	return BoneTransforms;
+}
+
+TArray<FVector> ULiveLinkDataPreviewComponent::GetLocatorPositions() const
+{
+	ILiveLinkClient& LiveLinkClient = IModularFeatures::Get().GetModularFeature<ILiveLinkClient>(ILiveLinkClient::ModularFeatureName);
+	
+	TArray<FVector> MarkerPostions;
+
+	TSubclassOf<ULiveLinkRole> SubjectRole = LiveLinkClient.GetSubjectRole_AnyThread(SubjectName);
+	if (SubjectRole)
+	{
+		if(LiveLinkClient.DoesSubjectSupportsRole_AnyThread(SubjectName, ULiveLinkLocatorRole::StaticClass()))
+		{
+			if(LiveLinkClient.IsSubjectValid(SubjectName))
+			{
+				FLiveLinkSubjectFrameData SubjectFrameData;
+				LiveLinkClient.EvaluateFrame_AnyThread(SubjectName, LiveLinkClient.GetSubjectRole_AnyThread(SubjectName), SubjectFrameData);
+							
+				const FLiveLinkLocatorStaticData* MarkerStaticData = SubjectFrameData.StaticData.Cast<FLiveLinkLocatorStaticData>();
+				const FLiveLinkLocatorFrameData* MarkerFrameData = SubjectFrameData.FrameData.Cast<FLiveLinkLocatorFrameData>();
+				check(MarkerStaticData);
+				check(MarkerFrameData);
+
+				if (MarkerFrameData->Locators.Num() > 0)
+				{
+					TArray<FVector> MarkerPositions = MarkerFrameData->Locators;
+					return MarkerFrameData->Locators;
+				}
+			}
+
+		}
+	}
+	return MarkerPostions;
+}
+
+FTransform ULiveLinkDataPreviewComponent::GetSingleTransform() const
+{
+	ILiveLinkClient& LiveLinkClient = IModularFeatures::Get().GetModularFeature<ILiveLinkClient>(ILiveLinkClient::ModularFeatureName);
+
+	FTransform Transform;
+	
+	if(LiveLinkClient.DoesSubjectSupportsRole_AnyThread(SubjectName, ULiveLinkTransformRole::StaticClass()))
+	{
+		if(LiveLinkClient.IsSubjectValid(SubjectName))
+		{
+			FLiveLinkSubjectFrameData SubjectFrameData;
+			LiveLinkClient.EvaluateFrame_AnyThread(SubjectName, LiveLinkClient.GetSubjectRole_AnyThread(SubjectName), SubjectFrameData);
+			const FLiveLinkTransformFrameData* TransformFrameData = SubjectFrameData.FrameData.Cast<FLiveLinkTransformFrameData>();
+			Transform = TransformFrameData->Transform;
+		}
+
+	}
+
+	return Transform;
+}
+
+void ULiveLinkDataPreviewComponent::CacheSkeletalAnimationData(const FLiveLinkSkeletonStaticData* InStaticData, const FLiveLinkAnimationFrameData* InFrameData)
+{
+	{
+		CachedSkeletonData = *InStaticData;
+		CachedAnimationData = *InFrameData;
+
+		const int32 NumTransforms = InFrameData->Transforms.Num();
+		check(InStaticData->BoneNames.Num() == NumTransforms);
+		check(InStaticData->BoneParents.Num() == NumTransforms);
+		check(InStaticData->PropertyNames.Num() == InFrameData->PropertyValues.Num());
+		CachedRootSpaceTransforms.SetNum(NumTransforms);
+		CachedChildTransformIndices.SetNum(NumTransforms);
+		for (int32 i = 0; i < NumTransforms; ++i)
+		{
+			CachedRootSpaceTransforms[i].Key = false;
+			CachedChildTransformIndices[i].Key = false;
+		}
+	}
+}
+
+bool ULiveLinkDataPreviewComponent::IsValidTransformIndex(int32 InTransformIndex, FLiveLinkAnimationFrameData InAnimData)
+{
+	return (InTransformIndex >= 0) && (InTransformIndex < InAnimData.Transforms.Num());
+}
diff --git a/Engine/Plugins/VirtualProduction/PerformanceCaptureWorkflow/Source/PerformanceCaptureWorkflow/Private/Visualizers/LiveLinkDataPreviewComponent.h b/Engine/Plugins/VirtualProduction/PerformanceCaptureWorkflow/Source/PerformanceCaptureWorkflow/Private/Visualizers/LiveLinkDataPreviewComponent.h
new file mode 100644
--- /dev/null
+++ b/Engine/Plugins/VirtualProduction/PerformanceCaptureWorkflow/Source/PerformanceCaptureWorkflow/Private/Visualizers/LiveLinkDataPreviewComponent.h
@@ -0,0 +1,95 @@
+// Copyright Epic Games, Inc. All Rights Reserved.
+
+#pragma once
+
+#include "CoreMinimal.h"
+#include "LiveLinkTypes.h"
+#include "Components/InstancedStaticMeshComponent.h"
+#include "Roles/LiveLinkAnimationTypes.h"
+#include "Roles/LiveLinkLocatorTypes.h"
+#include "Roles/LiveLinkTransformTypes.h"
+#include "LiveLinkDataPreviewComponent.generated.h"
+
+class FCanvas;
+class UFont;
+class FPrimitiveDrawInterface;
+class FSceneView;
+
+UENUM(BlueprintType)
+enum class ELiveLinkVisualBoneType : uint8
+{
+	Joint = 0		UMETA(DisplayName = "Joint", Toolip = "For use drawing the location of a joint"),
+	Bone = 1		UMETA(DisplayName = "Bone", Toolip = "For use drawing a bone pointing at the next child in a chain")
+};
+
+/**
+ * An instance static mesh component for drawing LiveLink subject data in-level. Support drawing Transforms, Locators, Skeletons and Cameras
+ */
+UCLASS(ClassGroup=(LiveLink), Transient, DisplayName = "Live Link Data Preview")
+class PERFORMANCECAPTUREWORKFLOW_API ULiveLinkDataPreviewComponent : public UInstancedStaticMeshComponent
+{
+	GENERATED_BODY()
+
+	public:
+
+	//Sets default values for this component's properties
+	ULiveLinkDataPreviewComponent(const FObjectInitializer& ObjectInitializer);
+
+	//Called every frame
+	virtual void TickComponent(float DeltaTime, ELevelTick TickType, FActorComponentTickFunction* ThisTickFunction) override;
+
+	UPROPERTY(EditAnywhere, BlueprintReadWrite, Category = LiveLink)
+	FLiveLinkSubjectName SubjectName;
+
+	UPROPERTY(EditAnywhere, BlueprintReadOnly, Category=LiveLink)
+	bool bEvaluateLiveLink;
+
+	UPROPERTY(EditAnywhere, BlueprintReadOnly, Category="PerformanceCapture|Visualization")
+	bool bDrawLabels;
+
+	UPROPERTY(BlueprintReadWrite, EditAnywhere, Category="PerformanceCapture|Visualization")
+	ELiveLinkVisualBoneType BoneVisualType = ELiveLinkVisualBoneType::Bone;
+
+	UFUNCTION(BlueprintCallable, Category="PerformanceCapture|Visualization")
+	void SetEvaluateLiveLinkData(bool bNewEvaluateLiveLink) { bEvaluateLiveLink = bNewEvaluateLiveLink; }
+
+	UFUNCTION(BlueprintCallable, Category="PerformanceCapture|Visualization")
+	void SetDrawLabels(bool bNewDrawLabel) { bDrawLabels = bNewDrawLabel; }
+	
+	bool bIsDirty;
+	
+#if WITH_EDITOR
+
+	virtual void PostEditChangeProperty(FPropertyChangedEvent& PropertyChangedEvent) override;
+	
+#endif
+	
+protected:
+	
+	virtual void OnRegister() override;
+	virtual void OnUnregister() override;
+	
+	void CreateInstances();
+	
+	TArray<FTransform> GetJointTransforms() const;
+	TArray<FTransform> GetBoneTransforms() const;
+	TArray<FVector> GetLocatorPositions() const;
+	FTransform GetSingleTransform() const;;
+
+private:
+	
+	TArray<FName> MarkerLabels;
+
+	void CacheSkeletalAnimationData(const FLiveLinkSkeletonStaticData* InStaticData, const FLiveLinkAnimationFrameData* InFrameData);
+	
+	FLiveLinkSkeletonStaticData CachedSkeletonData;
+	FLiveLinkAnimationFrameData CachedAnimationData;
+
+	
+	void GetTransformRootSpace(const int32 InTransformIndex, FTransform& OutTransform) const;
+
+	mutable TArray<TPair<bool, FTransform>> CachedRootSpaceTransforms;
+	mutable TArray<TPair<bool, TArray<int32>>> CachedChildTransformIndices;
+	static bool IsValidTransformIndex(int32 InTransformIndex, FLiveLinkAnimationFrameData AnimData);
+	
+};
diff --git a/Engine/Plugins/VirtualProduction/PerformanceCaptureWorkflow/Source/PerformanceCaptureWorkflow/Private/Visualizers/PCapBoneVisualizer.cpp b/Engine/Plugins/VirtualProduction/PerformanceCaptureWorkflow/Source/PerformanceCaptureWorkflow/Private/Visualizers/PCapBoneVisualizer.cpp
new file mode 100644
--- /dev/null
+++ b/Engine/Plugins/VirtualProduction/PerformanceCaptureWorkflow/Source/PerformanceCaptureWorkflow/Private/Visualizers/PCapBoneVisualizer.cpp
@@ -0,0 +1,148 @@
+// Copyright Epic Games, Inc. All Rights Reserved.
+
+
+#include "PCapBoneVisualizer.h"
+#include "AnimationCoreLibrary.h"
+#include "SkeletalMeshAttributes.h"
+#include "Animation/SkeletalMeshActor.h"
+#include "Components/SkeletalMeshComponent.h"
+#include "Materials/MaterialInstanceDynamic.h"
+
+
+// Sets default values for this component's properties
+UPCapBoneVisualiser::UPCapBoneVisualiser(const FObjectInitializer& ObjectInitializer)
+{
+	PrimaryComponentTick.bCanEverTick = true;
+	bTickInEditor = true; //make this component tick in editor
+	VisualizationType = EBoneVisualType::Joint;
+	Color = FLinearColor(0.5,0.5,0.5,1.0);
+	UPrimitiveComponent::SetCollisionEnabled(ECollisionEnabled::NoCollision);
+	bAlwaysCreatePhysicsState = false;
+	SetCastShadow(false);
+}
+
+void UPCapBoneVisualiser::OnRegister()
+{
+	Super::OnRegister();
+
+	ASkeletalMeshActor* OwnerActor = Cast<ASkeletalMeshActor>(GetOwner());
+	
+	if(OwnerActor)
+	{
+		SkelmeshComponent = OwnerActor->GetSkeletalMeshComponent();
+	}
+
+	ClearInstances();
+
+	if(SkelmeshComponent)
+	{
+		TArray<FTransform> BoneTransforms = GetBoneTransforms(SkelmeshComponent);
+		AddInstances(BoneTransforms, false, true, false);
+
+		if(!DynamicMaterial)
+		{
+			UMaterialInterface* BaseMaterial = GetMaterial(0);
+			if(BaseMaterial) //Check if the BaseMaterial is Valid. It will only be valid if a mesh has already been sent. 
+			{
+				DynamicMaterial = CreateDynamicMaterialInstance(0, BaseMaterial, TEXT("MaterialInstance"));
+				SetMaterial(0, DynamicMaterial);
+				DynamicMaterial->ClearParameterValues();
+				DynamicMaterial->SetVectorParameterValue(FName(TEXT("Color")), Color);
+			}
+		}
+	}
+}
+
+void UPCapBoneVisualiser::OnUnregister()
+{
+	Super::OnUnregister();
+
+	ClearInstances();
+}
+
+void UPCapBoneVisualiser::TickComponent(float DeltaTime, ELevelTick TickType, FActorComponentTickFunction* ThisTickFunction)
+{
+	Super::TickComponent(DeltaTime, TickType, ThisTickFunction);
+	
+	if(SkelmeshComponent)
+	{
+		switch (VisualizationType)
+		{
+		case EBoneVisualType::Bone:
+			{
+				BatchUpdateInstancesTransforms(0, GetBoneTransforms(SkelmeshComponent), true, true, false);
+			}
+			break;
+		case EBoneVisualType::Joint:
+			{
+				BatchUpdateInstancesTransforms(0, GetJointTransforms(SkelmeshComponent), true, true, false);
+			}
+
+			break;
+		default:
+			checkNoEntry();
+		}
+	}
+}
+
+TArray<FTransform> UPCapBoneVisualiser::GetJointTransforms(USkeletalMeshComponent* InSkelMeshComponent) const
+{
+	TArray<FTransform> FoundTransforms;
+
+	int32 i = InSkelMeshComponent->GetNumBones();
+	FTransform BoneTransform;
+
+	for (i = 0; i < InSkelMeshComponent->GetNumBones(); i++ )
+	{
+		BoneTransform = InSkelMeshComponent->GetBoneTransform(InSkelMeshComponent->GetBoneName(i));
+		FoundTransforms.Add(BoneTransform);
+	}
+
+	return FoundTransforms;
+}
+
+TArray<FTransform> UPCapBoneVisualiser::GetBoneTransforms(USkeletalMeshComponent* InSkelMeshComponent) const
+{
+	TArray<FTransform> FoundTransforms;
+
+	int32 i = InSkelMeshComponent->GetNumBones();
+	FTransform BoneTransform;
+	FTransform ParentBoneTransform;
+	double BoneLength;
+	FVector Aim = FVector( 0, 0, 1);
+
+	for (i = 0; i < InSkelMeshComponent->GetNumBones(); i++ )
+	{
+		BoneTransform = InSkelMeshComponent->GetBoneTransform(InSkelMeshComponent->GetBoneName(i));
+		ParentBoneTransform = InSkelMeshComponent->GetBoneTransform(InSkelMeshComponent->GetParentBone(InSkelMeshComponent->GetBoneName(i)));
+		BoneLength = FVector::Distance(BoneTransform.GetLocation(), ParentBoneTransform.GetLocation());
+		BoneTransform.SetScale3D(FVector(1, 1, BoneLength));
+		FQuat DiffRotation = AnimationCore::SolveAim(BoneTransform, ParentBoneTransform.GetLocation(), Aim.GetSafeNormal(), false, FVector(1, 1, 1), float(0));
+		BoneTransform.SetRotation(DiffRotation);
+		FoundTransforms.Add(BoneTransform);
+	}
+	return FoundTransforms;
+}
+
+void UPCapBoneVisualiser::UpdateColor(FLinearColor NewColor)
+{
+	Color = NewColor;
+	if(DynamicMaterial)
+	{
+		DynamicMaterial->SetVectorParameterValue(FName(TEXT("Color")), Color);
+	}
+}
+
+void UPCapBoneVisualiser::PostEditChangeProperty(FPropertyChangedEvent& PropertyChangedEvent)
+{
+	Super::PostEditChangeProperty(PropertyChangedEvent);
+
+	if(PropertyChangedEvent.Property != nullptr)
+	{
+		if(PropertyChangedEvent.Property->GetFName() == GET_MEMBER_NAME_CHECKED(UPCapBoneVisualiser, Color))
+		{
+			UpdateColor(Color);
+		}
+	}
+}
+
diff --git a/Engine/Plugins/VirtualProduction/PerformanceCaptureWorkflow/Source/PerformanceCaptureWorkflow/Private/Visualizers/PCapBoneVisualizer.h b/Engine/Plugins/VirtualProduction/PerformanceCaptureWorkflow/Source/PerformanceCaptureWorkflow/Private/Visualizers/PCapBoneVisualizer.h
new file mode 100644
--- /dev/null
+++ b/Engine/Plugins/VirtualProduction/PerformanceCaptureWorkflow/Source/PerformanceCaptureWorkflow/Private/Visualizers/PCapBoneVisualizer.h
@@ -0,0 +1,64 @@
+// Copyright Epic Games, Inc. All Rights Reserved.
+
+#pragma once
+
+#include "CoreMinimal.h"
+#include "Components/InstancedStaticMeshComponent.h"
+#include "PCapBoneVisualizer.generated.h"
+
+class USkeletalMeshComponent;
+class UMaterialInstanceDynamic;
+
+UENUM(BlueprintType)
+enum class EBoneVisualType : uint8
+{
+	Joint	UMETA(DisplayName = "Joint"),
+	Bone	UMETA(DisplayName = "Bone"),
+};
+
+/** Instanced Static Mesh Component for drawing bones and joints on Skeletal Mesh. Only usable with SkeletalMesh Actors*/
+UCLASS(ClassGroup=(PerformanceCapture), meta=(BlueprintSpawnableComponent),Transient, DisplayName = "PerformanceCaptureVisualizer")
+class PERFORMANCECAPTUREWORKFLOW_API UPCapBoneVisualiser : public UInstancedStaticMeshComponent
+{
+	GENERATED_BODY()
+
+public:
+	// Sets default values for this component's properties
+	UPCapBoneVisualiser(const FObjectInitializer& ObjectInitializer);
+	
+	// Called every frame
+	virtual void TickComponent(float DeltaTime, ELevelTick TickType, FActorComponentTickFunction* ThisTickFunction) override;
+
+	/** Get all the transforms from the skelmesh*/
+	TArray<FTransform> GetJointTransforms(USkeletalMeshComponent* InSkelMeshComponent) const;
+
+	/** Get all the bones' transforms for drawing a joint (point to point mesh)*/
+	TArray<FTransform> GetBoneTransforms(USkeletalMeshComponent* InSkelMeshComponent) const;
+	
+	/** Whether to draw joints along the length of the bone, or just at the pivot or scaled by the distance to the parent bone*/
+	UPROPERTY(EditAnywhere, BlueprintReadWrite, Category="PerformanceCapture|Visualization")
+	EBoneVisualType VisualizationType;
+
+	UPROPERTY(EditAnywhere, BlueprintReadWrite, Setter="UpdateColor", Category="PerformanceCapture|Visualization")
+	FLinearColor Color;
+	
+	UPROPERTY(EditAnywhere, BlueprintReadOnly, Category="PerformanceCapture|Visualization")
+	TObjectPtr<UMaterialInstanceDynamic> DynamicMaterial;
+
+	UFUNCTION()
+	void UpdateColor(FLinearColor NewColor);
+
+#if WITH_EDITOR
+
+	virtual void PostEditChangeProperty(FPropertyChangedEvent& PropertyChangedEvent) override;
+	
+#endif
+
+protected:
+	virtual void OnRegister() override;
+	virtual void OnUnregister() override;
+
+private:
+	UPROPERTY()
+	TObjectPtr<USkeletalMeshComponent> SkelmeshComponent;
+};
diff --git a/Engine/Plugins/VirtualProduction/PerformanceCaptureWorkflow/Source/PerformanceCaptureWorkflow/Private/Visualizers/PCapMarkerVisualizer.cpp b/Engine/Plugins/VirtualProduction/PerformanceCaptureWorkflow/Source/PerformanceCaptureWorkflow/Private/Visualizers/PCapMarkerVisualizer.cpp
new file mode 100644
--- /dev/null
+++ b/Engine/Plugins/VirtualProduction/PerformanceCaptureWorkflow/Source/PerformanceCaptureWorkflow/Private/Visualizers/PCapMarkerVisualizer.cpp
@@ -0,0 +1,132 @@
+// Copyright Epic Games, Inc. All Rights Reserved.
+
+
+#include "PCapMarkerVisualizer.h"
+#include "Materials/MaterialInstanceDynamic.h"
+#include "Roles/LiveLinkAnimationRole.h"
+#include "ILiveLinkClient.h"
+#include "PerformanceCapture.h"
+#include "Features/IModularFeatures.h"
+#include "Roles/LiveLinkAnimationRole.h"
+#include "Roles/LiveLinkTransformRole.h"
+
+UPCapMarkerVisualizer::UPCapMarkerVisualizer(const FObjectInitializer& ObjectInitializer)
+{
+	PrimaryComponentTick.bCanEverTick = true;
+	bTickInEditor = true; //make this component tick in editor
+	PrimaryComponentTick.TickGroup = ETickingGroup::TG_PostUpdateWork;
+	bDrawLabels = false;
+	XAxisSign = EAxisSign::Positive;
+	YAxisSign = EAxisSign::Positive;
+	bEvaluateLiveLink = true;
+	
+}
+
+UPCapMarkerVisualizer::~UPCapMarkerVisualizer()
+{
+	//Placeholder 
+}
+
+void UPCapMarkerVisualizer::TickComponent(float DeltaTime, ELevelTick TickType, FActorComponentTickFunction* ThisTickFunction)
+{
+	Super::TickComponent(DeltaTime, TickType, ThisTickFunction);
+
+	if(GetInstanceCount() == 0)
+	{
+		AddInstances(GetMarkerTransforms(), true,true,false);
+	}
+
+	if(bEvaluateLiveLink)
+	{
+		BatchUpdateInstancesTransforms(0, GetMarkerTransforms(), true,true,false);
+	}
+	
+	if(bDrawLabels)
+	{
+		DrawLabels(MarkerLabels);
+	}
+}
+
+TArray<FTransform> UPCapMarkerVisualizer::GetMarkerTransforms()
+{
+	FLiveLinkSubjectFrameData SubjectFrameData;
+	
+	ILiveLinkClient& LiveLinkClient = IModularFeatures::Get().GetModularFeature<ILiveLinkClient>(ILiveLinkClient::ModularFeatureName);
+
+	LiveLinkClient.EvaluateFrame_AnyThread(LiveLinkSubject, LiveLinkClient.GetSubjectRole_AnyThread(LiveLinkSubject), SubjectFrameData);
+
+	const FLiveLinkBaseFrameData* MarkerData = SubjectFrameData.FrameData.GetBaseData();
+	const FLiveLinkBaseStaticData* MarkerNames = SubjectFrameData.StaticData.GetBaseData();
+
+	//TODO : special handling for world attached skeletal mesh actor (animation role)
+
+	if(MarkerData)
+	{
+		//Early out if the two array's aren't the same length
+		if(MarkerNames->PropertyNames.Num() != MarkerData->PropertyValues.Num())
+		{
+			UE_LOG(LogPCap, Warning, TEXT("LiveLnk data not valid - length of property values does not match length of property names"))
+			return MarkerLocations;
+		}
+
+		MarkerLocations.Empty();
+	
+		const TArray<float>& LocationValues = MarkerData->PropertyValues;
+		float XScale = GetAxisSign(XAxisSign);
+		float YScale = GetAxisSign(YAxisSign);
+		FTransform MarkerTransform;
+		AActor* OwnerActor = GetOwner();
+		AActor* AttachedParent = GetOwner()->GetAttachParentActor();
+
+		if(LiveLinkClient.GetSubjectRole_AnyThread(LiveLinkSubject) == ULiveLinkAnimationRole::StaticClass())
+		{
+			for ( int i=0; i < MarkerData->PropertyValues.Num();  i = i + 3)
+			{
+				MarkerTransform = FTransform(FVector(LocationValues[i], LocationValues[i+1], LocationValues[i+2])) * FTransform(FRotator(0,0,0), FVector(0,0,0),FVector(XScale,YScale,1));
+				MarkerTransform = MarkerTransform * OwnerActor->GetActorTransform();
+				//MarkerLocations.Emplace(MarkerTransform);
+				MarkerLocations.Add(MarkerTransform);
+			}
+		}
+		else //If not an animation role assume a transform or a basic role
+		{
+			for ( int i=0; i < MarkerData->PropertyValues.Num();  i = i + 3)
+			{
+				if(AttachedParent!=nullptr)
+				{
+					MarkerTransform = FTransform(FVector(LocationValues[i], LocationValues[i+1], LocationValues[i+2])) * FTransform(FRotator(0,0,0), FVector(0,0,0),FVector(XScale,YScale,1));
+					MarkerTransform = MarkerTransform * AttachedParent->GetActorTransform();
+					MarkerLocations.Add(MarkerTransform);
+				}
+				else
+				{
+					MarkerTransform = FTransform(FVector(LocationValues[i], LocationValues[i+1], LocationValues[i+2])) * FTransform(FRotator(0,0,0), FVector(0,0,0),FVector(XScale,YScale,1));
+					MarkerLocations.Add(MarkerTransform);
+				}
+			}
+		}
+		
+	}
+	return MarkerLocations;
+}
+
+void UPCapMarkerVisualizer::DrawLabels(TArray<FName> Labels)
+{
+	//ToDo
+}
+
+void UPCapMarkerVisualizer::OnRegister()
+{
+	Super::OnRegister();
+
+	//ClearInstances();
+	
+	AddInstances(GetMarkerTransforms(), true,true,false);
+}
+
+void UPCapMarkerVisualizer::OnUnregister()
+{
+	Super::OnUnregister();
+
+	ClearInstances();
+}
diff --git a/Engine/Plugins/VirtualProduction/PerformanceCaptureWorkflow/Source/PerformanceCaptureWorkflow/Private/Visualizers/PCapMarkerVisualizer.h b/Engine/Plugins/VirtualProduction/PerformanceCaptureWorkflow/Source/PerformanceCaptureWorkflow/Private/Visualizers/PCapMarkerVisualizer.h
new file mode 100644
--- /dev/null
+++ b/Engine/Plugins/VirtualProduction/PerformanceCaptureWorkflow/Source/PerformanceCaptureWorkflow/Private/Visualizers/PCapMarkerVisualizer.h
@@ -0,0 +1,83 @@
+// Copyright Epic Games, Inc. All Rights Reserved.
+
+#pragma once
+
+#include "CoreMinimal.h"
+#include "Components/InstancedStaticMeshComponent.h"
+#include "LiveLinkTypes.h"
+#include "ILiveLinkClient.h"
+#include "PCapMarkerVisualizer.generated.h"
+
+class UMaterialInstanceDynamic;
+
+UENUM(BlueprintType)
+enum class EAxisSign : uint8
+{
+	Positive = 0	UMETA(DisplayName = "Positive"),
+	Negative = 1	UMETA(DisplayName = "Negative")
+};
+
+/**
+ * An Instanced Static Mesh Component to represent Motion Capture marker data locations
+ */
+UCLASS(ClassGroup=(PerformanceCapture), meta=(BlueprintSpawnableComponent),Transient, DisplayName = "MocapMarkerVisualizer")
+class PERFORMANCECAPTUREWORKFLOW_API UPCapMarkerVisualizer : public UInstancedStaticMeshComponent
+{
+	GENERATED_BODY()
+
+public:
+	//Sets default values for this component's properties
+	UPCapMarkerVisualizer(const FObjectInitializer& ObjectInitializer);
+
+	//Destructor
+	~UPCapMarkerVisualizer();
+
+	//Called every frame
+	virtual void TickComponent(float DeltaTime, ELevelTick TickType, FActorComponentTickFunction* ThisTickFunction) override;
+
+	/** Get marker transforms defined in mocap tracking space */
+	TArray<FTransform> GetMarkerTransforms();
+
+	static void DrawLabels(TArray<FName> Labels);
+
+	/*UPROPERTY(EditAnywhere, BlueprintReadWrite, Category="PerformanceCapture|Visualization")
+	EMarkerVisualType VisualizationType;*/
+	
+	UPROPERTY(EditAnywhere, BlueprintReadOnly, Category="PerformanceCapture|Visualization")
+	TObjectPtr<UMaterialInstanceDynamic> DynamicMaterial;
+
+	UPROPERTY(BlueprintReadWrite, EditAnywhere, Category="PerformanceCapture|Visualization")
+	FLiveLinkSubjectName LiveLinkSubject;
+
+	//To Do - add text labels to drawn markers
+	UPROPERTY(BlueprintReadWrite, EditAnywhere, Category="PerformanceCapture|Visualization")
+	bool bEvaluateLiveLink;
+
+	//To Do - add text labels to drawn markers
+	UPROPERTY(BlueprintReadWrite, EditAnywhere, Category="PerformanceCapture|Visualization")
+	bool bDrawLabels;
+	
+	UPROPERTY(BlueprintReadWrite, EditAnywhere, Category="PerformanceCapture|Visualization")
+	EAxisSign XAxisSign;
+
+	UPROPERTY(BlueprintReadWrite, EditAnywhere, Category="PerformanceCapture|Visualization")
+	EAxisSign YAxisSign;
+	
+	UPROPERTY(BlueprintReadWrite, EditAnywhere, Category="PerformanceCapture|Visualization")
+	TArray<FTransform> MarkerLocations;
+
+protected:
+	
+	virtual void OnRegister() override;
+	virtual void OnUnregister() override;
+
+private:
+	
+	TArray<FName> MarkerLabels;
+
+	static float GetAxisSign(EAxisSign AxisSign)
+	{
+		switch (AxisSign)
+		{case EAxisSign::Positive:return 1.0f; case EAxisSign::Negative:return -1.0f; default: return 1.0f;	}
+	}
+};
diff --git a/Engine/Plugins/VirtualProduction/PerformanceCaptureWorkflow/Source/PerformanceCaptureWorkflow/Public/PCapMVVMResolver.h b/Engine/Plugins/VirtualProduction/PerformanceCaptureWorkflow/Source/PerformanceCaptureWorkflow/Public/PCapMVVMResolver.h
new file mode 100644
--- /dev/null
+++ b/Engine/Plugins/VirtualProduction/PerformanceCaptureWorkflow/Source/PerformanceCaptureWorkflow/Public/PCapMVVMResolver.h
@@ -0,0 +1,23 @@
+// Copyright Epic Games, Inc. All Rights Reserved.
+
+#pragma once
+
+#include "CoreMinimal.h"
+#include "View/MVVMViewModelContextResolver.h"
+#include "PCapSubsystem.h"
+#include "PCapMVVMResolver.generated.h"
+
+/**
+ * Editor-only resolver class
+ */
+UCLASS()
+class UPCapMVVMResolver : public UMVVMViewModelContextResolver
+{
+	GENERATED_BODY()
+
+#if WITH_EDITOR
+	
+	virtual UObject* CreateInstance(const UClass* ExpectedType, const UUserWidget* UserWidget, const UMVVMView* View) const override;
+
+#endif
+};
diff --git a/Engine/Plugins/VirtualProduction/PerformanceCaptureWorkflow/Source/PerformanceCaptureWorkflow/Public/PCapStageRoot.h b/Engine/Plugins/VirtualProduction/PerformanceCaptureWorkflow/Source/PerformanceCaptureWorkflow/Public/PCapStageRoot.h
new file mode 100644
--- /dev/null
+++ b/Engine/Plugins/VirtualProduction/PerformanceCaptureWorkflow/Source/PerformanceCaptureWorkflow/Public/PCapStageRoot.h
@@ -0,0 +1,30 @@
+// Copyright Epic Games, Inc. All Rights Reserved.
+
+
+#pragma once
+
+#include "CoreMinimal.h"
+#include "GameFramework/Actor.h"
+#include "PCapStageRoot.generated.h"
+
+class USceneCaptureComponent2D;
+
+UCLASS(Blueprintable)
+class APerformanceCaptureStageRoot : public AActor
+{
+	GENERATED_BODY()
+
+public:
+	// Sets default values for this actor's properties
+	APerformanceCaptureStageRoot(const FObjectInitializer& ObjectInitializer);
+	
+
+	UFUNCTION(BlueprintImplementableEvent, BlueprintCallable, Category="Performance Capture|StageRoot")
+	void SetGridVisiblity(bool NewVisibility);
+
+	UFUNCTION(BlueprintImplementableEvent, BlueprintCallable, Category="Performance Capture|StageRoot")
+	void SetStageGhostVisiblity(bool NewVisibility);
+	
+	UPROPERTY(EditAnywhere, BlueprintReadOnly,  Category="Peformance Capture|StageRoot")
+	TObjectPtr<USceneCaptureComponent2D> MapCaptureComponent;
+};
\ No newline at end of file
diff --git a/Engine/Plugins/VirtualProduction/PerformanceCaptureWorkflow/Source/PerformanceCaptureWorkflow/Public/PCapViewmodel.h b/Engine/Plugins/VirtualProduction/PerformanceCaptureWorkflow/Source/PerformanceCaptureWorkflow/Public/PCapViewmodel.h
new file mode 100644
--- /dev/null
+++ b/Engine/Plugins/VirtualProduction/PerformanceCaptureWorkflow/Source/PerformanceCaptureWorkflow/Public/PCapViewmodel.h
@@ -0,0 +1,20 @@
+// Copyright Epic Games, Inc. All Rights Reserved.
+
+#pragma once
+
+#include "CoreMinimal.h"
+#include "MVVMViewModelBase.h"
+#include "PCapViewmodel.generated.h"
+
+/**
+ * 
+ */
+UCLASS(Blueprintable, Abstract, DisplayName="PCap Viewmodel",meta = (ShowWorldContextPin))
+class PERFORMANCECAPTUREWORKFLOW_API UPCapViewmodel : public UMVVMViewModelBase
+{
+	GENERATED_BODY()
+public:
+
+	UFUNCTION(BlueprintCallable, BlueprintImplementableEvent, Category="Performance Capture|Viewmodel")
+	void Initialize();
+};
diff --git a/Engine/Plugins/VirtualProduction/PerformanceCaptureWorkflow/Source/PerformanceCaptureWorkflow/Public/PerformanceCapture.h b/Engine/Plugins/VirtualProduction/PerformanceCaptureWorkflow/Source/PerformanceCaptureWorkflow/Public/PerformanceCapture.h
new file mode 100644
--- /dev/null
+++ b/Engine/Plugins/VirtualProduction/PerformanceCaptureWorkflow/Source/PerformanceCaptureWorkflow/Public/PerformanceCapture.h
@@ -0,0 +1,42 @@
+// Copyright Epic Games, Inc. All Rights Reserved.
+
+#pragma once
+
+#include "CoreMinimal.h"
+#include "IPlacementModeModule.h"
+#include "Modules/ModuleManager.h"
+
+DECLARE_LOG_CATEGORY_EXTERN(LogPCap, Log, All);
+class FPlacementModeID;
+
+class FPerformanceCaptureModule : public IModuleInterface
+{
+public:
+
+	/** IModuleInterface implementation */
+	virtual void StartupModule() override;
+	virtual void ShutdownModule() override;
+	
+	/** This function will be bound to Command (by default it will bring up plugin window) */
+	void PluginButtonClicked();
+
+private:
+
+	void RegisterMenus();
+
+	/** Register items to show up in the Place Actors panel. */
+	void RegisterPlacementModeItems();
+
+	/** Unregister items in Place Actors panel */
+	void UnregisterPlacementModeItems();
+
+	/** Gathers the Info on the Virtual Production Place Actors Category */
+	const FPlacementCategoryInfo* GetVirtualProductionCategoryRegisteredInfo() const;
+
+	TArray<TOptional<FPlacementModeID>> PlaceActors;
+	FDelegateHandle PostEngineInitHandle;
+
+	TSharedRef<class SDockTab> OnSpawnMocapManager(const class FSpawnTabArgs& SpawnTabArgs);
+	
+	TSharedPtr<class FUICommandList> PluginCommands;
+};
diff --git a/Engine/Plugins/VirtualProduction/PerformanceCaptureWorkflow/Source/PerformanceCaptureWorkflow/Public/PerformanceCaptureCommands.h b/Engine/Plugins/VirtualProduction/PerformanceCaptureWorkflow/Source/PerformanceCaptureWorkflow/Public/PerformanceCaptureCommands.h
new file mode 100644
--- /dev/null
+++ b/Engine/Plugins/VirtualProduction/PerformanceCaptureWorkflow/Source/PerformanceCaptureWorkflow/Public/PerformanceCaptureCommands.h
@@ -0,0 +1,23 @@
+// Copyright Epic Games, Inc. All Rights Reserved.
+
+#pragma once
+
+#include "CoreMinimal.h"
+#include "Framework/Commands/Commands.h"
+#include "PerformanceCaptureStyle.h"
+
+class FPerformanceCaptureCommands : public TCommands<FPerformanceCaptureCommands>
+{
+public:
+
+	FPerformanceCaptureCommands()
+		: TCommands<FPerformanceCaptureCommands>(TEXT("PerformanceCapture"), NSLOCTEXT("Contexts", "PerformanceCapture", "PerformanceCaptureWorkflow Plugin"), NAME_None, FPerformanceCaptureStyle::GetStyleSetName())
+	{
+	}
+
+	// TCommands<> interface
+	virtual void RegisterCommands() override;
+
+public:
+	TSharedPtr< FUICommandInfo > OpenPluginWindow;
+};
\ No newline at end of file
diff --git a/Engine/Plugins/VirtualProduction/PerformanceCaptureWorkflow/Source/PerformanceCaptureWorkflow/Public/PerformanceCaptureStyle.h b/Engine/Plugins/VirtualProduction/PerformanceCaptureWorkflow/Source/PerformanceCaptureWorkflow/Public/PerformanceCaptureStyle.h
new file mode 100644
--- /dev/null
+++ b/Engine/Plugins/VirtualProduction/PerformanceCaptureWorkflow/Source/PerformanceCaptureWorkflow/Public/PerformanceCaptureStyle.h
@@ -0,0 +1,36 @@
+// Copyright Epic Games, Inc. All Rights Reserved.
+
+#pragma once
+
+#include "CoreMinimal.h"
+#include "Styling/SlateStyle.h"
+
+class ISlateStyle;
+struct FSlateBrush;
+
+/**  */
+class FPerformanceCaptureStyle
+{
+public:
+
+	static void Initialize();
+
+	static void Shutdown();
+
+	/** reloads textures used by slate renderer */
+	static void ReloadTextures();
+
+	/** @return The Slate style set for the editor */
+	static const ISlateStyle& Get();
+
+	static FName GetStyleSetName();
+
+	static FColor TypeColor;
+
+
+private:
+
+	static TSharedRef< class FSlateStyleSet > Create();
+	
+	static TSharedPtr< class FSlateStyleSet > StyleInstance;
+};
