PCap Workflow
- Marked Stage Root base class as abstract so users can't accidently place the native class in level.
- Added validation to the spawn stage button so only enabled when a valid stage root class is set.

#virtualized

diff --git a/Engine/Plugins/VirtualProduction/PerformanceCaptureWorkflow/Source/PerformanceCaptureWorkflow/Public/PCapStageRoot.h b/Engine/Plugins/VirtualProduction/PerformanceCaptureWorkflow/Source/PerformanceCaptureWorkflow/Public/PCapStageRoot.h
--- a/Engine/Plugins/VirtualProduction/PerformanceCaptureWorkflow/Source/PerformanceCaptureWorkflow/Public/PCapStageRoot.h
+++ b/Engine/Plugins/VirtualProduction/PerformanceCaptureWorkflow/Source/PerformanceCaptureWorkflow/Public/PCapStageRoot.h
@@ -9,7 +9,7 @@
 
 class USceneCaptureComponent2D;
 
-UCLASS(Blueprintable)
+UCLASS(Blueprintable, Abstract)
 class APerformanceCaptureStageRoot : public AActor
 {
 	GENERATED_BODY()
