[PCap Workflow]

- Fix Localization errors in PCapNamingTokens

diff --git a/Engine/Plugins/VirtualProduction/PerformanceCaptureWorkflow/Source/PerformanceCaptureWorkflow/Private/PCapNamingTokens.cpp b/Engine/Plugins/VirtualProduction/PerformanceCaptureWorkflow/Source/PerformanceCaptureWorkflow/Private/PCapNamingTokens.cpp
--- a/Engine/Plugins/VirtualProduction/PerformanceCaptureWorkflow/Source/PerformanceCaptureWorkflow/Private/PCapNamingTokens.cpp
+++ b/Engine/Plugins/VirtualProduction/PerformanceCaptureWorkflow/Source/PerformanceCaptureWorkflow/Private/PCapNamingTokens.cpp
@@ -39,7 +39,7 @@ void UPCapNamingTokens::OnCreateDefaultTokens(TArray<FNamingTokenData>& Tokens)
 	
 	Tokens.Add( {
 		TEXT("sessionToken"),
-		LOCTEXT("PCapTokenSession", "Session"),
+		LOCTEXT("PCapTokenSessionToken", "SessionToken"),
 		FNamingTokenData::FTokenProcessorDelegateNative::CreateLambda([this]
 		{
 			if (Context && Context->SessionTemplate)
@@ -72,7 +72,7 @@ void UPCapNamingTokens::OnCreateDefaultTokens(TArray<FNamingTokenData>& Tokens)
 	
 	Tokens.Add( {
 		TEXT("pcapRootFolder"),
-		LOCTEXT("PCapTokenProduction", "RootFolder"),
+		LOCTEXT("PCapTokenRootFolder", "RootFolder"),
 		FNamingTokenData::FTokenProcessorDelegateNative::CreateLambda([this]
 		{
 			if (Context)
@@ -88,7 +88,7 @@ void UPCapNamingTokens::OnCreateDefaultTokens(TArray<FNamingTokenData>& Tokens)
 		
 	Tokens.Add( {
 		TEXT("sessionFolder"),
-		LOCTEXT("PCapTokenProduction", "SessionFolder"),
+		LOCTEXT("PCapTokenSessionFolder", "SessionFolder"),
 		FNamingTokenData::FTokenProcessorDelegateNative::CreateLambda([this]
 		{
 			if (Context)
