[PCap Workflow]

- Added multicast BP delegate for whenever an asset editor is opened or closed, emits the opened/closed asset as a ref. This  has been done primarily to allow editor utility to keep track of if/which level sequence is open and when it is closed
- Reorder Take status enum so sorting by byte works logically - thumbs up, down then neutral
-Takes record struct now uses TSoftObjectPtr for pointer to saved level sequences. Hard-ref used previously was causing sync load hitches.

#rnx

diff --git a/Engine/Plugins/VirtualProduction/PerformanceCaptureWorkflow/Source/PerformanceCaptureWorkflow/Private/PCapDatabase.h b/Engine/Plugins/VirtualProduction/PerformanceCaptureWorkflow/Source/PerformanceCaptureWorkflow/Private/PCapDatabase.h
--- a/Engine/Plugins/VirtualProduction/PerformanceCaptureWorkflow/Source/PerformanceCaptureWorkflow/Private/PCapDatabase.h
+++ b/Engine/Plugins/VirtualProduction/PerformanceCaptureWorkflow/Source/PerformanceCaptureWorkflow/Private/PCapDatabase.h
@@ -66,9 +66,9 @@ public:
 UENUM(BlueprintType)
 enum class EPCapTakeStatus : uint8
 {
-	Neutral,
 	ThumbsUp,
-	ThumbsDown
+	ThumbsDown,
+	Neutral
 };
 
 USTRUCT(BlueprintType)
@@ -84,7 +84,7 @@ public:
 	
 	/** The level sequence recorded */
 	UPROPERTY(EditAnywhere, BlueprintReadWrite, Category="Performance Capture Take")
-	TObjectPtr<ULevelSequence> RecordedTake;
+	TSoftObjectPtr<ULevelSequence> RecordedTake;
 	
 	/** Seconds Duration */
 	UPROPERTY(VisibleAnywhere, BlueprintReadWrite, Category="Performance Capture Take")
diff --git a/Engine/Plugins/VirtualProduction/PerformanceCaptureWorkflow/Source/PerformanceCaptureWorkflow/Private/PCapSubsystem.cpp b/Engine/Plugins/VirtualProduction/PerformanceCaptureWorkflow/Source/PerformanceCaptureWorkflow/Private/PCapSubsystem.cpp
--- a/Engine/Plugins/VirtualProduction/PerformanceCaptureWorkflow/Source/PerformanceCaptureWorkflow/Private/PCapSubsystem.cpp
+++ b/Engine/Plugins/VirtualProduction/PerformanceCaptureWorkflow/Source/PerformanceCaptureWorkflow/Private/PCapSubsystem.cpp
@@ -10,10 +10,9 @@
 #include "Misc/CoreDelegates.h"
 #include "UObject/Package.h"
 #include "Editor.h"
+#include "Subsystems/AssetEditorSubsystem.h"
 
 
-class FAssetRegistryModule;
-
 void UPerformanceCaptureSubsystem::Initialize(FSubsystemCollectionBase& Collection)
 {
 	Super::Initialize(Collection);
@@ -47,14 +46,19 @@ void UPerformanceCaptureSubsystem::OnEngineInitComplete()
 	AssetRegistry->OnAssetAdded().AddUObject(this, &UPerformanceCaptureSubsystem::OnAssetAdded);
 	FCoreUObjectDelegates::OnObjectPropertyChanged.AddUObject(this, &UPerformanceCaptureSubsystem::OnObjectPropertyChanged);
 
+
 #if WITH_EDITOR
 	if (GEditor)
+	{
+		if (UTransBuffer* TransactionBuffer = Cast<UTransBuffer>(GEditor->Trans))
 		{
-			if (UTransBuffer* TransactionBuffer = Cast<UTransBuffer>(GEditor->Trans))
-			{
-				TransactionBuffer->OnUndo().AddUObject(this, &UPerformanceCaptureSubsystem::OnEditorUndo);
-				TransactionBuffer->OnRedo().AddUObject(this, &UPerformanceCaptureSubsystem::OnEditorRedo);
-			}
+			TransactionBuffer->OnUndo().AddUObject(this, &UPerformanceCaptureSubsystem::OnEditorUndo);
+			TransactionBuffer->OnRedo().AddUObject(this, &UPerformanceCaptureSubsystem::OnEditorRedo);
+		}
+
+		GEditor->GetEditorSubsystem<UAssetEditorSubsystem>()->OnAssetOpenedInEditor().AddUObject(this, &UPerformanceCaptureSubsystem::OnAssetOpened);
+		GEditor->GetEditorSubsystem<UAssetEditorSubsystem>()->OnAssetClosedInEditor().AddUObject(this, &UPerformanceCaptureSubsystem::OnAssetClosed);
+
 	}
 #endif
 	
@@ -125,6 +129,14 @@ void UPerformanceCaptureSubsystem::OnEditorRedo(const FTransactionContext& Trans
 	}
 }
 
+void UPerformanceCaptureSubsystem::OnAssetOpened(UObject* Object, IAssetEditorInstance* Instance) const
+{
+	OnPCapAssetEditorOpen.Broadcast(Object);
+}
 
+void UPerformanceCaptureSubsystem::OnAssetClosed(UObject* Object, IAssetEditorInstance* Instance) const
+{
+	OnPCapAssetEditorClose.Broadcast(Object);
+}
 
 
diff --git a/Engine/Plugins/VirtualProduction/PerformanceCaptureWorkflow/Source/PerformanceCaptureWorkflow/Private/PCapSubsystem.h b/Engine/Plugins/VirtualProduction/PerformanceCaptureWorkflow/Source/PerformanceCaptureWorkflow/Private/PCapSubsystem.h
--- a/Engine/Plugins/VirtualProduction/PerformanceCaptureWorkflow/Source/PerformanceCaptureWorkflow/Private/PCapSubsystem.h
+++ b/Engine/Plugins/VirtualProduction/PerformanceCaptureWorkflow/Source/PerformanceCaptureWorkflow/Private/PCapSubsystem.h
@@ -9,11 +9,11 @@
 #include "AssetRegistry/AssetData.h"
 #include "UObject/Object.h"
 
-
 #include "PCapSubsystem.generated.h"
 
 class UPerformanceCaptureDatabaseHelper;
 class UMVVMViewModelCollectionObject;
+class IAssetEditorInstance;
 
 #if WITH_EDITOR
 struct FTransactionContext;
@@ -76,6 +76,16 @@ public:
 	UPROPERTY(BlueprintAssignable, Category="Performance Capture|Workflow")
 	FOnPCapEditorRedo OnPCapEditorRedo;
 
+	DECLARE_DYNAMIC_MULTICAST_DELEGATE_OneParam(FOnPCapAssetEditorOpen, UObject*, Asset);
+	/** Editor Only*/
+	UPROPERTY(BlueprintAssignable, Category="Performance Capture|Assets")
+	FOnPCapAssetEditorOpen OnPCapAssetEditorOpen;
+
+	DECLARE_DYNAMIC_MULTICAST_DELEGATE_OneParam(FOnPCapAssetEditorClose, UObject*, Asset);
+	/** Editor Only*/
+	UPROPERTY(BlueprintAssignable, Category="Performance Capture|Assets")
+	FOnPCapAssetEditorClose OnPCapAssetEditorClose;
+
 
 #endif
 
@@ -111,6 +121,10 @@ private:
 	void OnEditorUndo(const FTransactionContext& TransactionContext, bool Succeeded) const;
 
 	void OnEditorRedo(const FTransactionContext& TransactionContext, bool Succeeded) const;
+
+	void OnAssetOpened(UObject* Object, IAssetEditorInstance* Instance) const;
+	
+	void OnAssetClosed(UObject* Object, IAssetEditorInstance* Instance) const;
 };
 
 //Todo: add post edit change to re-int the viewmodel and database helper if the user changes these once the editor has started. 
