[PCapWorkflow]

-Removes static function to get a TakeRecorderActorSource's settings objects as this now duplicates a function in Take Recorder subsystem

diff --git a/Engine/Plugins/VirtualProduction/PerformanceCaptureWorkflow/Source/PerformanceCaptureWorkflow/Private/PCapBPFunctionLibrary.cpp b/Engine/Plugins/VirtualProduction/PerformanceCaptureWorkflow/Source/PerformanceCaptureWorkflow/Private/PCapBPFunctionLibrary.cpp
--- a/Engine/Plugins/VirtualProduction/PerformanceCaptureWorkflow/Source/PerformanceCaptureWorkflow/Private/PCapBPFunctionLibrary.cpp
+++ b/Engine/Plugins/VirtualProduction/PerformanceCaptureWorkflow/Source/PerformanceCaptureWorkflow/Private/PCapBPFunctionLibrary.cpp
@@ -56,18 +56,6 @@ FString UPerformanceCaptureBPFunctionLibrary::SanitizePathString(FString InStrin
 	}
 }
 
-TArray<UObject*> UPerformanceCaptureBPFunctionLibrary::GetTakeRecorderSourceSettingsObjects(UTakeRecorderSource* InSource, bool& bSuccess)
-{
-	if(InSource)
-	{
-		bSuccess = true;
-		return InSource->GetAdditionalSettingsObjects();
-	}
-	TArray<UObject*> NullResult;
-	bSuccess = false;
-	return NullResult;
-}
-
 void UPerformanceCaptureBPFunctionLibrary::GetAllActorsWithComponent(const UObject* WorldContextObject, TSubclassOf<UActorComponent> Component, TArray<AActor*>& OutActors)
 {
 	QUICK_SCOPE_CYCLE_COUNTER(ITWBPFactory_GetAllActorsWithComponent);
diff --git a/Engine/Plugins/VirtualProduction/PerformanceCaptureWorkflow/Source/PerformanceCaptureWorkflow/Private/PCapBPFunctionLibrary.h b/Engine/Plugins/VirtualProduction/PerformanceCaptureWorkflow/Source/PerformanceCaptureWorkflow/Private/PCapBPFunctionLibrary.h
--- a/Engine/Plugins/VirtualProduction/PerformanceCaptureWorkflow/Source/PerformanceCaptureWorkflow/Private/PCapBPFunctionLibrary.h
+++ b/Engine/Plugins/VirtualProduction/PerformanceCaptureWorkflow/Source/PerformanceCaptureWorkflow/Private/PCapBPFunctionLibrary.h
@@ -34,15 +34,6 @@ public:
 	 */
 	UFUNCTION(BlueprintCallable, Category = "Performance Capture|Workflow")
 	static FString SanitizePathString(FString InString);
-
-	/**
-	 * 
-	 * @param InSource TakeRecorderSourceObject
-	 * @param bSuccess Whether the Settings objects were successfully retrieved 
-	 * @return 
-	 */
-	UFUNCTION(BlueprintCallable, Category = "Performance Capture|Workflow")
-	static TArray<UObject*> GetTakeRecorderSourceSettingsObjects(UTakeRecorderSource* InSource,bool& bSuccess);
 	
 	/** 
 	*	Find all Actors in the world containing at least one instance of the given component class. 
