[PCap workflow]
- Improvements to Record Manager
- Now has validation message on lack of source
- Text entry validates for asset names (remove illegal characters)
- Validates take number against known takes, serves up warning if not available.
- Framerate is now a member of the viewmodel so it can be stored after the panel is removed and then retrieved later
- Added dependency to the IKRig plugin, even though it's enabled by default.

#rnx

#virtualized

diff --git a/Engine/Plugins/VirtualProduction/PerformanceCaptureWorkflow/PerformanceCaptureWorkflow.uplugin b/Engine/Plugins/VirtualProduction/PerformanceCaptureWorkflow/PerformanceCaptureWorkflow.uplugin
--- a/Engine/Plugins/VirtualProduction/PerformanceCaptureWorkflow/PerformanceCaptureWorkflow.uplugin
+++ b/Engine/Plugins/VirtualProduction/PerformanceCaptureWorkflow/PerformanceCaptureWorkflow.uplugin
@@ -1,12 +1,12 @@
 {
 	"FileVersion": 3,
 	"Version": 1,
-	"VersionName": "1.0",
+	"VersionName": "0.1",
 	"FriendlyName": "Performance Capture Workflow",
 	"Description": "Performance Capture In-Editor Workflow tools.",
 	"Category": "Animation",
 	"CreatedBy": "Epic Games",
-	"CreatedByURL": "",
+	"CreatedByURL": "https://epicgames.com",
 	"DocsURL": "",
 	"MarketplaceURL": "",
 	"SupportURL": "",
@@ -65,6 +65,10 @@
 		{
 			"Name": "NamingTokens",
 			"Enabled": true
+		},
+		{
+			"Name": "IKRig",
+			"Enabled": true
 		}
 	]
 }
\ No newline at end of file
