[Pcap Workflow]
- Adds BP delegate that fires when a LiveLink subject is updated

diff --git a/Engine/Plugins/VirtualProduction/PerformanceCaptureWorkflow/Source/PerformanceCaptureWorkflow/Private/PCapSubsystem.cpp b/Engine/Plugins/VirtualProduction/PerformanceCaptureWorkflow/Source/PerformanceCaptureWorkflow/Private/PCapSubsystem.cpp
--- a/Engine/Plugins/VirtualProduction/PerformanceCaptureWorkflow/Source/PerformanceCaptureWorkflow/Private/PCapSubsystem.cpp
+++ b/Engine/Plugins/VirtualProduction/PerformanceCaptureWorkflow/Source/PerformanceCaptureWorkflow/Private/PCapSubsystem.cpp
@@ -10,6 +10,8 @@
 #include "Misc/CoreDelegates.h"
 #include "UObject/Package.h"
 #include "Editor.h"
+#include "ILiveLinkClient.h"
+#include "Features/IModularFeatures.h"
 #include "Subsystems/AssetEditorSubsystem.h"
 
 
@@ -46,7 +48,6 @@ void UPerformanceCaptureSubsystem::OnEngineInitComplete()
 	AssetRegistry->OnAssetAdded().AddUObject(this, &UPerformanceCaptureSubsystem::OnAssetAdded);
 	FCoreUObjectDelegates::OnObjectPropertyChanged.AddUObject(this, &UPerformanceCaptureSubsystem::OnObjectPropertyChanged);
 
-
 #if WITH_EDITOR
 	if (GEditor)
 	{
@@ -58,11 +59,16 @@ void UPerformanceCaptureSubsystem::OnEngineInitComplete()
 
 		GEditor->GetEditorSubsystem<UAssetEditorSubsystem>()->OnAssetOpenedInEditor().AddUObject(this, &UPerformanceCaptureSubsystem::OnAssetOpened);
 		GEditor->GetEditorSubsystem<UAssetEditorSubsystem>()->OnAssetClosedInEditor().AddUObject(this, &UPerformanceCaptureSubsystem::OnAssetClosed);
-
 	}
 #endif
 	
-
+// Bind to LiveLink Subject changed delegate
+	IModularFeatures& ModularFeatures = IModularFeatures::Get();
+	if (ModularFeatures.IsModularFeatureAvailable(ILiveLinkClient::ModularFeatureName))
+	{
+		ILiveLinkClient& LiveLinkClient = IModularFeatures::Get().GetModularFeature<ILiveLinkClient>(ILiveLinkClient::ModularFeatureName);
+		LiveLinkClient.OnLiveLinkSubjectStateChanged().AddUObject(this, &UPerformanceCaptureSubsystem::OnLiveLinkSubjectUpdated);
+	}
 }
 
 void UPerformanceCaptureSubsystem::InitiateDatabaseHelper()
@@ -139,4 +145,12 @@ void UPerformanceCaptureSubsystem::OnAssetClosed(UObject* Object, IAssetEditorIn
 	OnPCapAssetEditorClose.Broadcast(Object);
 }
 
+void UPerformanceCaptureSubsystem::OnLiveLinkSubjectUpdated(FLiveLinkSubjectKey Subject, ELiveLinkSubjectState State) const
+{
+	OnPCapLiveLinkSubjectUpdate.Broadcast(Subject, State);
+}
+
+
+
+
 
diff --git a/Engine/Plugins/VirtualProduction/PerformanceCaptureWorkflow/Source/PerformanceCaptureWorkflow/Private/PCapSubsystem.h b/Engine/Plugins/VirtualProduction/PerformanceCaptureWorkflow/Source/PerformanceCaptureWorkflow/Private/PCapSubsystem.h
--- a/Engine/Plugins/VirtualProduction/PerformanceCaptureWorkflow/Source/PerformanceCaptureWorkflow/Private/PCapSubsystem.h
+++ b/Engine/Plugins/VirtualProduction/PerformanceCaptureWorkflow/Source/PerformanceCaptureWorkflow/Private/PCapSubsystem.h
@@ -3,6 +3,7 @@
 #pragma once
 
 #include "CoreMinimal.h"
+#include "LiveLinkTypes.h"
 #include "Subsystems/EngineSubsystem.h"
 #include "MVVMViewModelBase.h"
 #include "AssetRegistry/AssetRegistryModule.h"
@@ -86,6 +87,11 @@ public:
 	UPROPERTY(BlueprintAssignable, Category="Performance Capture|Assets")
 	FOnPCapAssetEditorClose OnPCapAssetEditorClose;
 
+	DECLARE_DYNAMIC_MULTICAST_DELEGATE_TwoParams(FOnPCapLiveLinkSubjectUpdate, FLiveLinkSubjectKey, Subject, ELiveLinkSubjectState, State);
+	/** Editor Only*/
+	UPROPERTY(BlueprintAssignable, Category="Performance Capture|LiveLink")
+	FOnPCapLiveLinkSubjectUpdate OnPCapLiveLinkSubjectUpdate;
+
 
 #endif
 
@@ -125,6 +131,8 @@ private:
 	void OnAssetOpened(UObject* Object, IAssetEditorInstance* Instance) const;
 	
 	void OnAssetClosed(UObject* Object, IAssetEditorInstance* Instance) const;
+
+	void OnLiveLinkSubjectUpdated(FLiveLinkSubjectKey Subject, ELiveLinkSubjectState State) const;
 };
 
 //Todo: add post edit change to re-int the viewmodel and database helper if the user changes these once the editor has started. 
