[PCap Workflow]

-Moved all of the stage root actor's component setup into C++ so users can make a BP child and it will give them everythig needed to link to the base functionality in the stage manager panel
-Updated child BPs, removed all BP functions that were duplicated
-Update session manager to handle the visibility switching and transactions.

#virtualized

diff --git a/Engine/Plugins/VirtualProduction/PerformanceCaptureWorkflow/Source/PerformanceCaptureWorkflow/Private/PCapStageRoot.cpp b/Engine/Plugins/VirtualProduction/PerformanceCaptureWorkflow/Source/PerformanceCaptureWorkflow/Private/PCapStageRoot.cpp
--- a/Engine/Plugins/VirtualProduction/PerformanceCaptureWorkflow/Source/PerformanceCaptureWorkflow/Private/PCapStageRoot.cpp
+++ b/Engine/Plugins/VirtualProduction/PerformanceCaptureWorkflow/Source/PerformanceCaptureWorkflow/Private/PCapStageRoot.cpp
@@ -2,7 +2,11 @@
 
 
 #include "PCapStageRoot.h"
-#include "Components/SceneCaptureComponent2D.h"
+#include "UObject/ConstructorHelpers.h"
+#include "Engine/TextureRenderTarget2D.h"
+#include "Materials/MaterialInterface.h"
+
+
 
 // Sets default values
 APerformanceCaptureStageRoot::APerformanceCaptureStageRoot(const FObjectInitializer& ObjectInitializer)
@@ -15,6 +19,31 @@ APerformanceCaptureStageRoot::APerformanceCaptureStageRoot(const FObjectInitiali
 	SetRootComponent(SceneRoot);
 	MapCaptureComponent = CreateDefaultSubobject<USceneCaptureComponent2D>("MapCaptureComponent");
 	MapCaptureComponent->SetupAttachment(GetRootComponent());
+	MapCaptureComponent->SetRelativeLocationAndRotation(FVector(0.0f, 0.0f, 500.0f), FRotator(-90.0f, 0.0f, -90.0f));
 	MapCaptureComponent->bCaptureEveryFrame = false;
 	MapCaptureComponent->bCaptureOnMovement = true;
+	MapCaptureComponent->CaptureSource = SCS_BaseColor;
+	MapCaptureComponent->ProjectionType = ECameraProjectionMode::Orthographic;
+	MapCaptureComponent->OrthoWidth = 1024.0f;
+	MapCaptureComponent->bUpdateOrthoPlanes = true;
+	
+	ConstructorHelpers::FObjectFinder<UTextureRenderTarget2D> MapTexture(TEXT("/PerformanceCaptureWorkflow/Stage/RT_OrthoView.RT_OrthoView"));
+	if (MapTexture.Succeeded())
+	{
+		MapCaptureComponent->TextureTarget = MapTexture.Object;
+	}
+
+	DecalComponent = CreateDefaultSubobject<UDecalComponent>("Decal");
+	DecalComponent->SetupAttachment(GetRootComponent());
+	DecalComponent->SetRelativeLocationAndRotation(FVector(0.0f, 0.0f, -100.0f), FRotator(0.0f, 0.0f, -90.0f));
+	ConstructorHelpers::FObjectFinder<UMaterialInterface> DecalMaterial(TEXT("/PerformanceCaptureWorkflow/Stage/MI_ProcGrid.MI_ProcGrid"));
+	if (DecalMaterial.Succeeded())
+	{
+		DecalComponent->SetDecalMaterial(DecalMaterial.Object);
+
+	}
+	DecalComponent->DecalSize = FVector(256.0f, 256.0f, 256.0f);
+
+	StageMeshParent = CreateDefaultSubobject<USceneComponent>("StageMeshParent");
+	StageMeshParent->SetupAttachment(GetRootComponent());
 }
\ No newline at end of file
diff --git a/Engine/Plugins/VirtualProduction/PerformanceCaptureWorkflow/Source/PerformanceCaptureWorkflow/Public/PCapStageRoot.h b/Engine/Plugins/VirtualProduction/PerformanceCaptureWorkflow/Source/PerformanceCaptureWorkflow/Public/PCapStageRoot.h
--- a/Engine/Plugins/VirtualProduction/PerformanceCaptureWorkflow/Source/PerformanceCaptureWorkflow/Public/PCapStageRoot.h
+++ b/Engine/Plugins/VirtualProduction/PerformanceCaptureWorkflow/Source/PerformanceCaptureWorkflow/Public/PCapStageRoot.h
@@ -5,6 +5,10 @@
 
 #include "CoreMinimal.h"
 #include "GameFramework/Actor.h"
+#include "Components/SceneComponent.h"
+#include "Components/DecalComponent.h"
+#include "Components/SceneCaptureComponent2D.h"
+
 #include "PCapStageRoot.generated.h"
 
 class USceneCaptureComponent2D;
@@ -18,13 +22,13 @@ public:
 	// Sets default values for this actor's properties
 	APerformanceCaptureStageRoot(const FObjectInitializer& ObjectInitializer);
 	
+	UPROPERTY(EditAnywhere, BlueprintReadWrite,  Category="Peformance Capture|StageRoot")
+	TObjectPtr<USceneCaptureComponent2D> MapCaptureComponent;
 
-	UFUNCTION(BlueprintImplementableEvent, BlueprintCallable, Category="Performance Capture|StageRoot")
-	void SetGridVisiblity(bool NewVisibility);
+	UPROPERTY(EditAnywhere, BlueprintReadWrite,  Category="Peformance Capture|StageRoot")
+	TObjectPtr<UDecalComponent> DecalComponent;
 
-	UFUNCTION(BlueprintImplementableEvent, BlueprintCallable, Category="Performance Capture|StageRoot")
-	void SetStageGhostVisiblity(bool NewVisibility);
+	UPROPERTY(EditAnywhere, BlueprintReadWrite,  Category="Peformance Capture|StageRoot", meta=(tooltip="You should parent your stage meshes under this component"))
+	TObjectPtr<USceneComponent> StageMeshParent;
 	
-	UPROPERTY(EditAnywhere, BlueprintReadOnly,  Category="Peformance Capture|StageRoot")
-	TObjectPtr<USceneCaptureComponent2D> MapCaptureComponent;
 };
\ No newline at end of file
