[PCap Workflow]

- Set visibility of Decal component (floor grid on stage root actor) to false by default.

#rnx

diff --git a/Engine/Plugins/VirtualProduction/PerformanceCaptureWorkflow/Source/PerformanceCaptureWorkflow/Private/PCapStageRoot.cpp b/Engine/Plugins/VirtualProduction/PerformanceCaptureWorkflow/Source/PerformanceCaptureWorkflow/Private/PCapStageRoot.cpp
--- a/Engine/Plugins/VirtualProduction/PerformanceCaptureWorkflow/Source/PerformanceCaptureWorkflow/Private/PCapStageRoot.cpp
+++ b/Engine/Plugins/VirtualProduction/PerformanceCaptureWorkflow/Source/PerformanceCaptureWorkflow/Private/PCapStageRoot.cpp
@@ -43,6 +43,7 @@ APerformanceCaptureStageRoot::APerformanceCaptureStageRoot(const FObjectInitiali
 
 	}
 	DecalComponent->DecalSize = FVector(256.0f, 256.0f, 256.0f);
+	DecalComponent->SetVisibility(false);
 
 	StageMeshParent = CreateDefaultSubobject<USceneComponent>("StageMeshParent");
 	StageMeshParent->SetupAttachment(GetRootComponent());
