[PCapWorkflow]

- Cleaned up comments on Stage Root actor

diff --git a/Engine/Plugins/VirtualProduction/PerformanceCaptureWorkflow/Source/PerformanceCaptureWorkflow/Private/PCapStageRoot.cpp b/Engine/Plugins/VirtualProduction/PerformanceCaptureWorkflow/Source/PerformanceCaptureWorkflow/Private/PCapStageRoot.cpp
--- a/Engine/Plugins/VirtualProduction/PerformanceCaptureWorkflow/Source/PerformanceCaptureWorkflow/Private/PCapStageRoot.cpp
+++ b/Engine/Plugins/VirtualProduction/PerformanceCaptureWorkflow/Source/PerformanceCaptureWorkflow/Private/PCapStageRoot.cpp
@@ -20,6 +20,8 @@ APerformanceCaptureStageRoot::APerformanceCaptureStageRoot(const FObjectInitiali
 	MapCaptureComponent = CreateDefaultSubobject<USceneCaptureComponent2D>("MapCaptureComponent");
 	MapCaptureComponent->SetupAttachment(GetRootComponent());
 	MapCaptureComponent->SetRelativeLocationAndRotation(FVector(0.0f, 0.0f, 500.0f), FRotator(-90.0f, 0.0f, -90.0f));
+	
+	// Set the default properties we want on the map scene capture. By default, do not run on every frame
 	MapCaptureComponent->bCaptureEveryFrame = false;
 	MapCaptureComponent->bCaptureOnMovement = true;
 	MapCaptureComponent->CaptureSource = SCS_BaseColor;
@@ -27,6 +29,7 @@ APerformanceCaptureStageRoot::APerformanceCaptureStageRoot(const FObjectInitiali
 	MapCaptureComponent->OrthoWidth = 1024.0f;
 	MapCaptureComponent->bUpdateOrthoPlanes = true;
 	
+	// Load and reference a render target in the plugin content
 	ConstructorHelpers::FObjectFinder<UTextureRenderTarget2D> MapTexture(TEXT("/PerformanceCaptureWorkflow/Stage/RT_OrthoView.RT_OrthoView"));
 	if (MapTexture.Succeeded())
 	{
@@ -36,15 +39,18 @@ APerformanceCaptureStageRoot::APerformanceCaptureStageRoot(const FObjectInitiali
 	DecalComponent = CreateDefaultSubobject<UDecalComponent>("Decal");
 	DecalComponent->SetupAttachment(GetRootComponent());
 	DecalComponent->SetRelativeLocationAndRotation(FVector(0.0f, 0.0f, -100.0f), FRotator(0.0f, 0.0f, -90.0f));
+	
+	// Load and reference the 1m grid decal material
 	ConstructorHelpers::FObjectFinder<UMaterialInterface> DecalMaterial(TEXT("/PerformanceCaptureWorkflow/Stage/MI_ProcGrid.MI_ProcGrid"));
 	if (DecalMaterial.Succeeded())
 	{
 		DecalComponent->SetDecalMaterial(DecalMaterial.Object);
-
 	}
+	
 	DecalComponent->DecalSize = FVector(256.0f, 256.0f, 256.0f);
 	DecalComponent->SetVisibility(false);
 
+	// Create a scene component to which all the stage meshes can be parented
 	StageMeshParent = CreateDefaultSubobject<USceneComponent>("StageMeshParent");
 	StageMeshParent->SetupAttachment(GetRootComponent());
 }
\ No newline at end of file
diff --git a/Engine/Plugins/VirtualProduction/PerformanceCaptureWorkflow/Source/PerformanceCaptureWorkflow/Public/PCapStageRoot.h b/Engine/Plugins/VirtualProduction/PerformanceCaptureWorkflow/Source/PerformanceCaptureWorkflow/Public/PCapStageRoot.h
--- a/Engine/Plugins/VirtualProduction/PerformanceCaptureWorkflow/Source/PerformanceCaptureWorkflow/Public/PCapStageRoot.h
+++ b/Engine/Plugins/VirtualProduction/PerformanceCaptureWorkflow/Source/PerformanceCaptureWorkflow/Public/PCapStageRoot.h
@@ -28,7 +28,7 @@ public:
 	UPROPERTY(EditAnywhere, BlueprintReadWrite,  Category="Peformance Capture|StageRoot")
 	TObjectPtr<UDecalComponent> DecalComponent;
 
-	UPROPERTY(EditAnywhere, BlueprintReadWrite,  Category="Peformance Capture|StageRoot", meta=(tooltip="You should parent your stage meshes under this component"))
+	UPROPERTY(EditAnywhere, BlueprintReadWrite,  Category="Peformance Capture|StageRoot")
 	TObjectPtr<USceneComponent> StageMeshParent;
 	
 };
\ No newline at end of file
