[PCap Workflow]
- Switched modification functions on the pcap datatable class to use FDataTableEditorUtils to modify data table rows instead of the runtime functions for user convenience and to avoid issues with scoped transactions and multicast delgates.

#rnx

diff --git a/Engine/Plugins/VirtualProduction/PerformanceCaptureWorkflow/Source/PerformanceCaptureWorkflow/Private/PCapDataTable.cpp b/Engine/Plugins/VirtualProduction/PerformanceCaptureWorkflow/Source/PerformanceCaptureWorkflow/Private/PCapDataTable.cpp
--- a/Engine/Plugins/VirtualProduction/PerformanceCaptureWorkflow/Source/PerformanceCaptureWorkflow/Private/PCapDataTable.cpp
+++ b/Engine/Plugins/VirtualProduction/PerformanceCaptureWorkflow/Source/PerformanceCaptureWorkflow/Private/PCapDataTable.cpp
@@ -2,23 +2,62 @@
 
 
 #include "PCapDataTable.h"
-
 #include "PCapDatabase.h"
+#include "DataTableEditorUtils.h"
 
 UPCapDataTable::UPCapDataTable()
 {
 	RowStruct = FPCapRecordBase::StaticStruct();
 	bStripFromClientBuilds = true; //Prevent any Performance Capture data from cooking into client builds as this is editor-only tooling
-
-	OnDataTableChanged().AddUObject(this, &UPCapDataTable::OnDataTableModified);
+	
+	OnDataTableChanged().AddUObject(this, &UPCapDataTable::DataTableModified);
 }
 
 UPCapDataTable::~UPCapDataTable()
 {
-	
+
 }
 
-void UPCapDataTable::OnDataTableModified() const
+void UPCapDataTable::DataTableModified() const
 {
 	OnDatatableModified.Broadcast();
 }
+
+bool UPCapDataTable::RemoveTableRow(FName RowName)
+{
+	return FDataTableEditorUtils::RemoveRow(this, RowName);
+}
+
+bool UPCapDataTable::DuplicateTableRow(FName SourceRow, FName NewRow)
+{
+	if (FDataTableEditorUtils::DuplicateRow(this, SourceRow, NewRow) !=NULL)
+	{
+		return true;
+	}
+	return false;
+}
+
+bool UPCapDataTable::AddTableRow(FName NewRow)
+{
+	if(FDataTableEditorUtils::AddRow(this, NewRow) == nullptr)
+	{
+		return false;
+	}
+	return true;
+}
+
+bool UPCapDataTable::InsertTableRow(FName SelectedRow, FName NewRow, bool bAbove)
+{
+	ERowInsertionPosition Postion = ERowInsertionPosition::Above;
+	if(!bAbove)
+	{
+		Postion = ERowInsertionPosition::Below;	
+	}
+	
+	if(FDataTableEditorUtils::AddRowAboveOrBelowSelection(this, SelectedRow, NewRow, Postion)!=NULL)
+	{
+		return true;
+	}
+	return false;
+}
+
diff --git a/Engine/Plugins/VirtualProduction/PerformanceCaptureWorkflow/Source/PerformanceCaptureWorkflow/Private/PCapDataTable.h b/Engine/Plugins/VirtualProduction/PerformanceCaptureWorkflow/Source/PerformanceCaptureWorkflow/Private/PCapDataTable.h
--- a/Engine/Plugins/VirtualProduction/PerformanceCaptureWorkflow/Source/PerformanceCaptureWorkflow/Private/PCapDataTable.h
+++ b/Engine/Plugins/VirtualProduction/PerformanceCaptureWorkflow/Source/PerformanceCaptureWorkflow/Private/PCapDataTable.h
@@ -23,5 +23,40 @@ public:
 	FOnDatatableModified OnDatatableModified;
 
 	/** Emit a callback each time the datatable is modified*/
-	void OnDataTableModified() const;
+	void DataTableModified() const;
+
+	/**
+	 * Remove a given row from the datatable.
+	 * @param RowName The row name to remove.
+	 * @return bool Whether operation was successful.
+	 */
+	UFUNCTION(BlueprintCallable, Category="Performance Capture|Database")
+	bool RemoveTableRow(FName RowName);
+
+	/**
+	 * Duplicate a given new row to the datatable.
+	 * @param SourceRow The row to duplicate.
+	 * @param NewRow The unique name for the new row.
+	 * @return bool Whether duplication was successful.
+	 */
+	UFUNCTION(BlueprintCallable, Category="Performance Capture|Database")
+	bool DuplicateTableRow(FName SourceRow, FName NewRow);
+
+	/**
+	 * Add a given new row to the datatable.
+	 * @param NewRow The unique name for the new row.
+	 * @return bool Whether operation was successful.
+	 */
+	UFUNCTION(BlueprintCallable, Category="Performance Capture|Database")
+	bool AddTableRow(FName NewRow);
+	
+	/**
+	 * Insert a given new row in to the datatable, above or below the selected row.
+	 * @param SelectedRow The row to insert at.
+	 * @param NewRow The unique name for the new row.
+	 * @param bAbove Insert Above or Below the SelectedRow.
+	 * @return bool Whether inserting new row was successful.
+	 */
+	UFUNCTION(BlueprintCallable, Category="Performance Capture|Database")
+	bool InsertTableRow(FName SelectedRow, FName NewRow, bool bAbove);
 };
