[PCap Workflow]
-Pass of improving documentation and comments.

diff --git a/Engine/Plugins/VirtualProduction/PerformanceCaptureWorkflow/Source/PerformanceCaptureWorkflow/Private/AnimNode_LiveLinkProp.h b/Engine/Plugins/VirtualProduction/PerformanceCaptureWorkflow/Source/PerformanceCaptureWorkflow/Private/AnimNode_LiveLinkProp.h
--- a/Engine/Plugins/VirtualProduction/PerformanceCaptureWorkflow/Source/PerformanceCaptureWorkflow/Private/AnimNode_LiveLinkProp.h
+++ b/Engine/Plugins/VirtualProduction/PerformanceCaptureWorkflow/Source/PerformanceCaptureWorkflow/Private/AnimNode_LiveLinkProp.h
@@ -15,7 +15,7 @@ class ULiveLinkRole;
 struct FLiveLinkSubjectFrameData;
 
 /**
- * This animnode is exclusively for Mocap props - single bone skeleton meshes. Not exposed to the animation graph
+ * This animnode is exclusively for Mocap props - single bone skeleton meshes. Not exposed to the animation graph.
  */
 
 USTRUCT(BlueprintInternalUseOnly)
@@ -24,16 +24,20 @@ struct FAnimNode_LiveLinkProp : public FAnimNode_Base
 	GENERATED_BODY()
 
 public:
-	UPROPERTY(EditAnywhere, BlueprintReadWrite, Category = "Performance Capture|Prop")
+	/** Input pose */
+	UPROPERTY(EditAnywhere, Category = "Performance Capture|Prop")
 	FPoseLink InputPose;
 
-	UPROPERTY(EditAnywhere, BlueprintReadWrite, Category = "Performance Capture|Prop")
+	/** The Live Link subject to use. */
+	UPROPERTY(EditAnywhere, Category = "Performance Capture|Prop")
 	FLiveLinkSubjectName LiveLinkSubjectName;
 
-	UPROPERTY(EditAnywhere, BlueprintReadWrite, Transient, Category = "Performance Capture|Prop")
+	/** Bool to control evaluation of animation. */
+	UPROPERTY(EditAnywhere, Transient, Category = "Performance Capture|Prop")
 	bool bDoLiveLinkEvaluation = true;
 
-	UPROPERTY(EditAnywhere, BlueprintReadWrite, Category = "Performance Capture|Prop")
+	/* Transform to apply local space to offset to the incoming data. */
+	UPROPERTY(EditAnywhere, Category = "Performance Capture|Prop")
 	FTransform OffsetTransform = FTransform::Identity;
 
 	
@@ -58,7 +62,6 @@ private:
 
 	ILiveLinkClient* LiveLinkClient_AnyThread;
 
-	//Delta time - TODO: might not need this because there is no retarget class at play
 	float CachedDeltaTime;
 
 	TSharedPtr<FLiveLinkSubjectFrameData> CachedLiveLinkFrameData;
diff --git a/Engine/Plugins/VirtualProduction/PerformanceCaptureWorkflow/Source/PerformanceCaptureWorkflow/Private/PCapActorFactory.h b/Engine/Plugins/VirtualProduction/PerformanceCaptureWorkflow/Source/PerformanceCaptureWorkflow/Private/PCapActorFactory.h
--- a/Engine/Plugins/VirtualProduction/PerformanceCaptureWorkflow/Source/PerformanceCaptureWorkflow/Private/PCapActorFactory.h
+++ b/Engine/Plugins/VirtualProduction/PerformanceCaptureWorkflow/Source/PerformanceCaptureWorkflow/Private/PCapActorFactory.h
@@ -11,7 +11,7 @@ class ACapturePerformer;
 struct FTypedElementHandle;
 
 /**
- * 
+ * Actor factory for spawning Capture Character from the PCapCharacterDataAsset.
  */
 
 UCLASS()
diff --git a/Engine/Plugins/VirtualProduction/PerformanceCaptureWorkflow/Source/PerformanceCaptureWorkflow/Private/PCapBPFunctionLibrary.h b/Engine/Plugins/VirtualProduction/PerformanceCaptureWorkflow/Source/PerformanceCaptureWorkflow/Private/PCapBPFunctionLibrary.h
--- a/Engine/Plugins/VirtualProduction/PerformanceCaptureWorkflow/Source/PerformanceCaptureWorkflow/Private/PCapBPFunctionLibrary.h
+++ b/Engine/Plugins/VirtualProduction/PerformanceCaptureWorkflow/Source/PerformanceCaptureWorkflow/Private/PCapBPFunctionLibrary.h
@@ -19,27 +19,25 @@ class PERFORMANCECAPTUREWORKFLOW_API UPerformanceCaptureBPFunctionLibrary : publ
 public:
 	/**
 	 * Returns a string stripped of the following characters []{}1!"$£%^&*()+=;:?/\|'@#~
-	 * 
-	 * @param InString String you wish to clean
-	 * @return Cleaned string
+	 * @param	InString String you wish to clean
+	 * @return	Cleaned string
 	 */
 	UFUNCTION(BlueprintCallable, Category = "Performance Capture|Workflow")
 	static FString SanitizeFileString(FString InString);
 	
 	/**
 	 * Returns a string stripped of the following characters []{}1!"$£%^&*()+=;:?\|'@#~ Note filepath delimiters are not excluded
-	 * 
-	 * @param InString String you wish to clean
-	 * @return Cleaned string 
+	 * @param	InString String you wish to clean
+	 * @return	Cleaned string 
 	 */
 	UFUNCTION(BlueprintCallable, Category = "Performance Capture|Workflow")
 	static FString SanitizePathString(FString InString);
 	
 	/** 
-	*	Find all Actors in the world containing at least one instance of the given component class. 
-	*	This is a very slow operation, as it will search over every actor in the world.
-	*	@param	Component	Class of component to find. Must be specified or result array will be empty.
-	*	@param	OutActors	Output array of Actors of the specified class.
+	* Find all Actors in the world containing at least one instance of the given component class. 
+	* This is a very slow operation, as it will search over every actor in the world.
+	* @param	Component	Class of component to find. Must be specified or result array will be empty.
+	* @param	OutActors	Output array of Actors of the specified class.
 	 */
 	UFUNCTION(BlueprintCallable, Category="Performance Capture|Core",  meta=(WorldContext="WorldContextObject", DeterminesOutputType="Component", DynamicOutputParam="OutActors"))
 	static void GetAllActorsWithComponent(const UObject* WorldContextObject, TSubclassOf<UActorComponent> Component, TArray<AActor*>& OutActors);
diff --git a/Engine/Plugins/VirtualProduction/PerformanceCaptureWorkflow/Source/PerformanceCaptureWorkflow/Private/PCapDatabase.h b/Engine/Plugins/VirtualProduction/PerformanceCaptureWorkflow/Source/PerformanceCaptureWorkflow/Private/PCapDatabase.h
--- a/Engine/Plugins/VirtualProduction/PerformanceCaptureWorkflow/Source/PerformanceCaptureWorkflow/Private/PCapDatabase.h
+++ b/Engine/Plugins/VirtualProduction/PerformanceCaptureWorkflow/Source/PerformanceCaptureWorkflow/Private/PCapDatabase.h
@@ -42,27 +42,33 @@ struct FPCapRecordBase : public FTableRowBase
 
 	virtual void OnPostDataImport(const UDataTable* InDataTable, const FName InRowName, TArray<FString>& OutCollectedImportProblems) override;
 	
+	/** GUID of the production record struct */
 	UPROPERTY(VisibleAnywhere, BlueprintReadWrite, Category="Performance Capture", meta=(IgnoreForMemberInitializationTest))
 	FGuid UID;
 
+	/** Bool to control if a record should be considered archived so the UI can hide it from view. */
 	UPROPERTY(EditAnywhere, BlueprintReadWrite, Category="Performance Capture")
 	bool bIsArchived = false;
 	
 };
 
+/** Struct to hold the record of a production. */
 USTRUCT(BlueprintType)
 struct FPCapProductionRecord : public FPCapRecordBase
 {
 	GENERATED_BODY()
 
 public:
+	/** Name of the production. */
 	UPROPERTY(EditAnywhere, BlueprintReadWrite, Category="Performance Capture Production")
 	FName ProductionName;
 
+	/** Notes on the production. */
 	UPROPERTY(EditAnywhere, BlueprintReadWrite, Category="Performance Capture Production")
 	FString ProductionNotes;
 };
 
+/** Take status enumerator. Thumbs Up, Down and Neutral. */
 UENUM(BlueprintType)
 enum class EPCapTakeStatus : uint8
 {
@@ -71,6 +77,7 @@ enum class EPCapTakeStatus : uint8
 	Neutral
 };
 
+/** Struct to hold the record of a recorded take. */
 USTRUCT(BlueprintType)
 struct FPCapTakeRecord : public FPCapRecordBase
 {
@@ -78,47 +85,47 @@ struct FPCapTakeRecord : public FPCapRecordBase
 	
 public:
 
-	/** Temporary Flag. Transient, so will not be saved */
+	/** Temporary Flag. Transient, so will not be saved. */
 	UPROPERTY(Transient, BlueprintReadWrite, Category="Flag")
 	bool bFlag = false;
 	
-	/** The level sequence recorded */
+	/** The level sequence recorded. */
 	UPROPERTY(EditAnywhere, BlueprintReadWrite, Category="Performance Capture Take")
 	TSoftObjectPtr<ULevelSequence> RecordedTake;
 	
-	/** Seconds Duration */
+	/** Seconds Duration. */
 	UPROPERTY(VisibleAnywhere, BlueprintReadWrite, Category="Performance Capture Take")
 	float TakeDurationSeconds = 0;
 
-	/** Recorded Framerate*/
+	/** Recorded Framerate. */
 	UPROPERTY(VisibleAnywhere, BlueprintReadWrite, Category="Performance Capture Take")
 	FFrameRate Framerate = FFrameRate(30, 1);
 	
-	/** HHMMSSFF Duration*/
+	/** HHMMSSFF Duration. */
 	UPROPERTY(VisibleAnywhere, BlueprintReadWrite, Category="Performance Capture Take")
 	FTimecode TakeDurationTimecode;
 
-	/** Start Timecode*/
+	/** Start Timecode. */
 	UPROPERTY(VisibleAnywhere, BlueprintReadWrite, Category="Performance Capture Take")
 	FTimecode StartTimecode;
 
-	/** End Timecode*/
+	/** End Timecode. */
 	UPROPERTY(VisibleAnywhere, BlueprintReadWrite, Category="Performance Capture Take")
 	FTimecode EndTimecode;
 
-	/** Transform of the stage root when this recording was made */
+	/** Transform of the stage root when this recording was made. */
 	UPROPERTY(EditAnywhere, BlueprintReadWrite, Category="Performance Capture Take")
 	FTransform MocapStageRootTransform;
 
-	/** Status of Take - Thumbs Up, Thumbs Down or None */
+	/** Status of Take - Thumbs Up, Thumbs Down or None. */
 	UPROPERTY(EditAnywhere, BlueprintReadWrite, Category="Performance Capture Take")
 	EPCapTakeStatus TakeStatus = EPCapTakeStatus::Neutral;
 
-	/** 5-Star Rating */
+	/** 5-Star Rating. */
 	UPROPERTY(EditAnywhere, BlueprintReadWrite, Category="Performance Capture Take", meta=(ClampMin=0, ClampMax=5))
 	int32 Rating = 0;
 	
-	/** Does this recording LiveLinkSource tracks */
+	/** Does this recording LiveLinkSource tracks. */
 	UPROPERTY(EditAnywhere, BlueprintReadWrite, Category="Performance Capture Take")
 	bool bContainsLiveLinkSources = false;
 
@@ -126,7 +133,8 @@ public:
 	UPROPERTY(EditAnywhere, BlueprintReadWrite, Category="Performance Capture Take")
 	bool bLiveLinkPlotted = false;
 	
-	/** Has the animation recorded for this take been replaced with an external asset import.
+	/**
+	 * Has the animation recorded for this take been replaced with an external asset import.
 	 * This option will be positive if you use the import and upgrade workflow
 	 * bringing fbx recordings from your motion capture software in and
 	 * replacing the live recording.
@@ -134,10 +142,12 @@ public:
 	UPROPERTY(EditAnywhere, BlueprintReadWrite, Category="Performance Capture Take")
 	bool bExternallyReplaced = false;
 
+	/** GUID for the session used during the recording of this take. */
 	UPROPERTY(VisibleAnywhere, BlueprintReadWrite, Category="Performance Capture Take")
 	FGuid SessionUID;
 };
 
+/** Struct to hold a datatable record of a session. */
 USTRUCT(BlueprintType)
 struct FPCapSessionRecord : public FPCapRecordBase
 {
@@ -145,78 +155,100 @@ struct FPCapSessionRecord : public FPCapRecordBase
 	
 public:
 	
-	/**Name for this session*/
+	/** Name for this session */
 	UPROPERTY(VisibleAnywhere, BlueprintReadWrite, Category="Performance Capture Session")
 	FName SessionName;
 
+	/** Date and time session record was created. */
 	UPROPERTY(VisibleAnywhere, BlueprintReadWrite, Category="Performance Capture Session")
 	FDateTime SessionDateTime;
 
+	/** Notes for this session record. */
 	UPROPERTY(EditAnywhere, BlueprintReadWrite, Category="Performance Capture Session")
 	FString SessionNotes;
 
-	/**ID of Production this Session belongs to */
+	/** Name of Production this Session belongs to. */
 	UPROPERTY(VisibleAnywhere, BlueprintReadWrite, Category="Performance Capture Session")
 	FName ProductionName;
 
+	/** GUID of the Production record associated with this session record. */
 	UPROPERTY(VisibleAnywhere, BlueprintReadWrite, Category="Performance Capture Session")
 	FGuid ProductionUID;
 
+	/** The token generated name of the session. */
 	UPROPERTY(VisibleAnywhere, BlueprintReadWrite, Category="Performance Capture Session", meta = (DisplayName = "Session OutputName"))
 	FString SessionOutputName;
 
+	/** The content browser path to this session record's data. */
 	UPROPERTY(VisibleAnywhere, BlueprintReadWrite, Category="Performance Capture Session", meta = (DisplayName = "Session Folder", ContentDir))
 	FString SessionPath;
 
+	/** The content browser path to this session record's performer data. */
 	UPROPERTY(VisibleAnywhere, BlueprintReadWrite , Category="Performance Capture Session", meta = (DisplayName = "Performer Folder", ContentDir))
 	FString PerformerPath;
-
+	
+	/** The content browser path to this session record's character data. */
 	UPROPERTY(VisibleAnywhere, BlueprintReadWrite, Category="Performance Capture Session", meta = (DisplayName = "Character Folder", ContentDir))
 	FString CharacterPath;
 
+	/** The content browser path to this session record's recorded takes data. */
 	UPROPERTY(VisibleAnywhere, BlueprintReadWrite, Category="Performance Capture Session", meta = (DisplayName = "Takes Folder", ContentDir))
 	FString TakesPath;
-	
+
+	/** The content browser path to this session record's prop data. */
 	UPROPERTY(VisibleAnywhere, BlueprintReadWrite, Category="Performance Capture Session", meta = (DisplayName = "Prop Folder", ContentDir))
 	FString PropPath;
 
+	/** The content browser path to this session record's scene data. */
 	UPROPERTY(VisibleAnywhere, BlueprintReadWrite, Category="Performance Capture Session", meta = (DisplayName = "Scene Folder", ContentDir))
 	FString ScenePath;
 
+	/** The content browser path to this session record's common data. */
 	UPROPERTY(VisibleAnywhere, BlueprintReadWrite, Category="Performance Capture Session", meta = (DisplayName = "Common Folder", ContentDir))
 	FString CommonPath;
 
+	/** Array of paths to folders of additional data. These folders can be defined in by the Session Template dataasset. */
 	UPROPERTY(VisibleAnywhere, BlueprintReadWrite, Category="Performance Capture Session", meta = (DisplayName = "Additional Folders", ContentDir))
 	TArray<FString> AdditionalFolders;
 
+	/** Reference to the Takes datatable created for and used by this session record. */
 	UPROPERTY(VisibleAnywhere, BlueprintReadWrite, Category="Performance Capture Session")
 	TSoftObjectPtr<UDataTable> TakesDataTable;
 
+	/** Array of soft-refs to Performers spawned during this session. */
 	UPROPERTY(VisibleAnywhere, BlueprintReadWrite, Category="Performance Capture Session")
 	TArray<TSoftObjectPtr<UPCapPerformerDataAsset>> Performers;
 
+	/** Array of soft-refs to Characters spawned during this session. */
 	UPROPERTY(VisibleAnywhere, BlueprintReadWrite, Category="Performance Capture Session")
 	TArray<TSoftObjectPtr<UPCapCharacterDataAsset>> Characters;
 
+	/** Array of soft-refs to Props spawned during this session. */
 	UPROPERTY(VisibleAnywhere, BlueprintReadWrite, Category="Performance Capture Session")
 	TArray<TSoftObjectPtr<UPCapPropDataAsset>> Props;
 
+	/** Reference to a locked Session Template. This will be generated on session creation. */
 	UPROPERTY(VisibleAnywhere, BlueprintReadWrite, Category="Performance Capture Session")
 	TSoftObjectPtr<UPCapSessionTemplate> SessionTemplate;
 
+	/** Bool to determine if this session record while in a level streaming or world partition level. */
 	UPROPERTY(VisibleAnywhere, BlueprintReadWrite, Category="Performance Capture Session")
 	bool bIsWorldPartition = false;
 
+	/** Sub-level created for this session, if using a persistent level of the Level streaming type. */
 	UPROPERTY(VisibleAnywhere, BlueprintReadWrite, Category="Performance Capture Session")
 	TSoftObjectPtr<UWorld> SubLevel;
 
+	/** Datalayer created for this session record, if using a persistent level of the World Partition type.  */
 	UPROPERTY(VisibleAnywhere, BlueprintReadWrite, Category="Performance Capture Session")
 	TSoftObjectPtr<UDataLayerAsset> SessionDataLayer;
 
+	/** The slates datatable to use in Mocap Recorder for this session. */
 	UPROPERTY(EditAnywhere, BlueprintReadWrite, Category="Performance Capture Session", meta = (RequiredAssetDataTags = "RowStructure=/Script/PerformanceCaptureWorkflow.PCapSlateRecord"))
 	TSoftObjectPtr<UPCapDataTable> SessionSlateTable;
 };
 
+/** Slate Status. Can be Incomplete, Complete, Skip.*/
 UENUM(BlueprintType)
 enum class EPCapSlateStatus : uint8
 {
@@ -225,24 +257,25 @@ enum class EPCapSlateStatus : uint8
 	Skip		UMETA(Tooltip="Slate is marked to skip")
 };
 
+/** Slate struct. A slate record is a what the user wishes call a take and provides prior to recording */
 USTRUCT(BlueprintType)
 struct FPCapSlateRecord : public FPCapRecordBase
 {                                                                                                                                                                                                                                                                                                                                                                                                                                 
 	GENERATED_BODY()
 
-	//*Name of slate*/
+	/** String to provide Name of Slate.  */
 	UPROPERTY(EditAnywhere, BlueprintReadWrite, Category="Performance Capture Slate", meta=(DisplayName="Slate"))
 	FString Slate;
 	
-	//*Note on slate. This will be passed to Take Recorder*/
+	/** Note on slate. This will be passed to the Mocap Recorder when using the Mocap Manager panel. */
 	UPROPERTY(EditAnywhere, BlueprintReadWrite, Category="Performance Capture Slate", meta=(DisplayName="Note", MultiLine=true))
 	FString SlateNote;
 
-	//*Status of Slate*/
+	/** Status of Slate, defined by the status enum. */
 	UPROPERTY(EditAnywhere, BlueprintReadWrite, Category="Performance Capture Slate", meta=(DisplayName="Status"))
 	EPCapSlateStatus SlateStatus = EPCapSlateStatus::Incomplete;
 
-	//*UID of the session this Slate is used in*/
+	/** UID of the session this Slate is used in. */
 	UPROPERTY(VisibleAnywhere, BlueprintReadWrite, Category="Performance Capture Slate")
 	FGuid SessionUID;
 };
@@ -274,12 +307,12 @@ public:
 
 private:
 
-	/** Guid for disambiguating actors spawned by the data assets*/
+	/** Guid for disambiguating actors spawned by the data assets. Only editable from BP. */
 	UPROPERTY(BlueprintGetter="GetAssetUID", Category="Performance Capture|Admin", meta=(IgnoreForMemberInitializationTest))
 	FGuid AssetUID;
 };
 
-/** Mocap Performer DataAsset */
+/** Mocap Performer DataAsset. This class is intended to track and encapsulate the properties and assets that make up a performer */
 UCLASS(Blueprintable, BlueprintType)
 class UPCapPerformerDataAsset : public UPCapDataAsset
 {
@@ -293,19 +326,19 @@ public:
 	UPROPERTY(BlueprintReadWrite, EditDefaultsOnly, Category="Performance Capture Performer")
 	FName PerformerName;
 	
-	/** LiveLink Subject for this Performer */
+	/** LiveLink Subject for this Performer. */
 	UPROPERTY(BlueprintReadWrite, EditDefaultsOnly, Category="Performance Capture Performer")
 	FLiveLinkSubjectName LiveLinkSubject;
 	
-	/** Performer Actor class to use for this performer */
+	/** Performer Actor class to use for this performer. */
 	UPROPERTY(BlueprintReadWrite, EditDefaultsOnly, NoClear, Category="Performance Capture Performer")
 	TSoftClassPtr<ACapturePerformer> PerformerActorClass;
 	
-	/** Color for this performer */
+	/** Color for this performer. */
 	UPROPERTY(BlueprintReadWrite, EditDefaultsOnly, Category="Performance Capture Performer")
 	FLinearColor PerformerColor;
 	
-	/** Performer Base Skeleton */
+	/** Performer Base Skeleton. */
 	UPROPERTY(BlueprintReadWrite, EditAnywhere, Category="Performance Capture Performer")
 	TSoftObjectPtr<USkeletalMesh> BaseSkeletalMesh;
 	
@@ -313,12 +346,12 @@ public:
 	UPROPERTY(BlueprintReadWrite, EditAnywhere, Category="Performance Capture Performer")
 	TSoftObjectPtr<USkeletalMesh> PerformerProportionedMesh;
 
-	/** Performer IKRig */
+	/** Performer IKRig. */
 	UPROPERTY(BlueprintReadWrite, EditDefaultsOnly, Category="Performance Capture Performer")
 	TSoftObjectPtr<UIKRigDefinition> IKRig;
 };
 
-/** Mocap Character Data asset */
+/** Mocap Character Data asset. This class is intended to track and encapsulate the properties and assets that make up a character */
 UCLASS(Blueprintable, BlueprintType)
 class UPCapCharacterDataAsset : public UPCapDataAsset
 {
@@ -326,14 +359,14 @@ class UPCapCharacterDataAsset : public UPCapDataAsset
 	
 public:
 
-	/** Capture Character Constructor */
+	// Capture Character Constructor */
 	UPCapCharacterDataAsset();
 
-	/** Capture Character Name */
+	/** Capture Character Name. */
 	UPROPERTY(BlueprintReadWrite, EditAnywhere, Category="Performance Capture Character", meta=(Tooltip="Optional Name for Character. If blank DataAsset name will be used"))
 	FName CharacterName;
 
-	/** Source Performer Asset for this Character */
+	/** Source Performer Asset for this Character. */
 	UPROPERTY(BlueprintReadWrite, EditAnywhere, Category="Performance Capture Character")
 	TSoftObjectPtr<UPCapPerformerDataAsset> SourcePerformerAsset;
 
@@ -341,18 +374,21 @@ public:
 	UPROPERTY(BlueprintReadWrite, EditAnywhere, Category="Performance Capture Character")
 	TSoftClassPtr<ACaptureCharacter>  CaptureCharacterClass;
 
-	/** Main Skeletal Mesh Asset. This will be the root component of any characters spawned. A character will not be spawned if this asset reference is null or not valid */
+	/** Main Skeletal Mesh Asset. This will be the root component of any characters spawned.
+	 * A character will not be spawned if this asset reference is null or not valid.
+	 */
 	UPROPERTY(BlueprintReadWrite, EditAnywhere, Category="Performance Capture Character")
 	TSoftObjectPtr<USkeletalMesh> SkeletalMesh;
 	
-	/** IKRig Asset */
+	/** IKRig Asset. */
 	UPROPERTY(BlueprintReadWrite, EditAnywhere, Category="Performance Capture Character")
 	TSoftObjectPtr<UIKRigDefinition> IKRig;
 
-	/** Retarget Asset */
+	/** Retarget Asset. */
 	UPROPERTY(BlueprintReadWrite, EditAnywhere, Category="Performance Capture Character")
 	TSoftObjectPtr<UIKRetargeter> Retargeter;
 
+	/** Array of additional skeleltal mesh assets that will be spawned and parented to the root skeleteal mesh asset*/
 	UPROPERTY(BlueprintReadWrite, EditAnywhere, Category="Performance Capture Character", meta=(DisplayName="Additional Meshes", ToolTip="Additional meshes that will be added as child components"))
 	TArray<TSoftObjectPtr<USkeletalMesh>> AdditionalMeshes;
 };
@@ -364,14 +400,14 @@ class UPCapPropDataAsset : public UPCapDataAsset
 	GENERATED_BODY()
 	
 public:
-	/** Prop DataAsset Constructor */
+	// Prop DataAsset Constructor */
 	UPCapPropDataAsset();
 	
-	/** Prop Name*/
+	/** Prop Name. */
 	UPROPERTY(BlueprintReadWrite, EditAnywhere, Category="Performance Capture Prop", meta=(Tooltip="Optional Name for Prop. If blank LiveLinkSubject name will be used"))
 	FName PropName;
 	
-	/** Prop's LiveLinkSubject */
+	/** Prop's LiveLinkSubject. */
 	UPROPERTY(BlueprintReadWrite, EditDefaultsOnly, Category="Performance Capture Prop")
 	FLiveLinkSubjectName LiveLinkSubject;
 
@@ -391,9 +427,11 @@ public:
 	UPROPERTY(BlueprintReadWrite, EditAnywhere, Category="Performance Capture Prop", meta=(DisplayName="Custom Prop Class", EditCondition="bCanSetCustomPropClass", HideEditConditionToggle))
 	TSoftClassPtr<AActor> CustomPropClass;
 
+	/** Bool to control if this prop will be set to bHiddenInGame=True when spawned. Useful for proxy objects you don't want to see in game view.*/
 	UPROPERTY(BlueprintReadWrite, EditDefaultsOnly, Category="Performance Capture Prop", meta=(Tooltip="Choose True for props you want to hide in 'Game' view e.g. Proxies or integrators"))
 	bool bHiddenInGame = false;
 
+	/** bool to control edit condition of the Static and Skeletal mesh properties */
 	UPROPERTY(EditDefaultsOnly, Transient, Category="Performance Capture Prop", AdvancedDisplay, meta=(DisplayName="Reset Read-only Options", Tooltip="Reset all mesh options (removes read-only status on fields)"))
 	bool bClearEditConditions = false;
 
@@ -417,6 +455,10 @@ private:
 	void ValidateEditConditions();
 };
 
+/** Editor-only class to enable the user to determine if creating database records should be done using intenral datatable assets or connect
+ * to some external database.
+ */
+//TODO currently a stub and not yet implemented.
 #if WITH_EDITOR
 UCLASS(Blueprintable, Abstract, meta = (ShowWorldContextPin))
 class UPerformanceCaptureDatabaseHelper : public UObject
diff --git a/Engine/Plugins/VirtualProduction/PerformanceCaptureWorkflow/Source/PerformanceCaptureWorkflow/Private/PCapNamingTokens.h b/Engine/Plugins/VirtualProduction/PerformanceCaptureWorkflow/Source/PerformanceCaptureWorkflow/Private/PCapNamingTokens.h
--- a/Engine/Plugins/VirtualProduction/PerformanceCaptureWorkflow/Source/PerformanceCaptureWorkflow/Private/PCapNamingTokens.h
+++ b/Engine/Plugins/VirtualProduction/PerformanceCaptureWorkflow/Source/PerformanceCaptureWorkflow/Private/PCapNamingTokens.h
@@ -9,7 +9,7 @@
 class UPCapSessionTemplate;
 
 /**
- * 
+ * Naming tokens class for Pcap workflow
  */
 UCLASS(NotBlueprintable)
 class PERFORMANCECAPTUREWORKFLOW_API UPCapNamingTokens : public UNamingTokens
@@ -40,11 +40,7 @@ private:
 	TObjectPtr<UPCapNamingTokensContext> Context;
 };
 
-/*
- *Context Object that references the current session template and allows it to be passed to the naming token functions
- * 
- */
-
+/** Context Object that references the current session template and allows it to be passed to the naming token functions */
 UCLASS(MinimalAPI)
 
 class UPCapNamingTokensContext : public UObject
@@ -53,7 +49,7 @@ class UPCapNamingTokensContext : public UObject
 	
 public:
 	
-	//The current template in use. This all seems rather circular
+	//The current template in use. This is so the session template can pass a reference of itself. 
 	UPROPERTY(Transient)
 	TObjectPtr<const UPCapSessionTemplate> SessionTemplate;
 };
\ No newline at end of file
diff --git a/Engine/Plugins/VirtualProduction/PerformanceCaptureWorkflow/Source/PerformanceCaptureWorkflow/Private/PCapPropComponent.h b/Engine/Plugins/VirtualProduction/PerformanceCaptureWorkflow/Source/PerformanceCaptureWorkflow/Private/PCapPropComponent.h
--- a/Engine/Plugins/VirtualProduction/PerformanceCaptureWorkflow/Source/PerformanceCaptureWorkflow/Private/PCapPropComponent.h
+++ b/Engine/Plugins/VirtualProduction/PerformanceCaptureWorkflow/Source/PerformanceCaptureWorkflow/Private/PCapPropComponent.h
@@ -36,7 +36,7 @@ public:
 	bool bEvaluateLiveLink = true;
 
 	/**
-	 * Overrides all LiveLink data and offset transform. Must be manually set on possessable bindings.
+	 * Overrides all LiveLink data and the offset transform. Must be manually set on possessable bindings.
 	 */
 	UPROPERTY(BlueprintReadOnly, EditInstanceOnly, Category = "Performance Capture")
 	bool bIsControlledBySequencer = false;
@@ -71,6 +71,10 @@ public:
 	UFUNCTION(BlueprintCallable, Category = "Performance Capture")
 	void SetControlledComponent(USceneComponent* InComponent);
 
+	/**
+	 * Set the Live Link subject used by this prop component. 
+	 * @param Subject Live Link Subject Name
+	 */
 	UFUNCTION(BlueprintCallable, Category = "Performance Capture|Prop")
 	void SetLiveLinkSubject(FLiveLinkSubjectName Subject);
 	
diff --git a/Engine/Plugins/VirtualProduction/PerformanceCaptureWorkflow/Source/PerformanceCaptureWorkflow/Private/PCapSessionTemplate.h b/Engine/Plugins/VirtualProduction/PerformanceCaptureWorkflow/Source/PerformanceCaptureWorkflow/Private/PCapSessionTemplate.h
--- a/Engine/Plugins/VirtualProduction/PerformanceCaptureWorkflow/Source/PerformanceCaptureWorkflow/Private/PCapSessionTemplate.h
+++ b/Engine/Plugins/VirtualProduction/PerformanceCaptureWorkflow/Source/PerformanceCaptureWorkflow/Private/PCapSessionTemplate.h
@@ -68,23 +68,32 @@ public:
 
 	virtual void PostLoad() override;
 	
+	/** String for production name. Will be sanitized for illegal filesystem characters*/
 	UPROPERTY(EditAnywhere, Category="Template", BlueprintGetter = GetProductionName, BlueprintSetter = SetProductionName,  meta = (DisplayName = "Production Name", EditCondition="bIsEditable"))
 	FString ProductionName;
 
+	/** String for session name. Will be sanitized for illegal filesystem character */
 	UPROPERTY(EditAnywhere, Category="Template", BlueprintGetter = GetSessionName, BlueprintSetter = SetSessionName, meta = (DisplayName = "Session Name", EditCondition="bIsEditable"))
 	FString SessionName;
 
-
-		
+	/** The final name of the session, evaluated from the given tokens and static strings */
 	UPROPERTY(EditAnywhere, BlueprintReadOnly, Category="Template", meta = (DisplayName = "Session Token", ContentDir, EditCondition="bIsEditable", tooltip="The final name of our session, including any tokens you include in it's construction. You can use tokens to build this string, including the Session Name above"))
 	FPCapTokenisedString SessionToken;
 
+	/**
+	 * Get the Session Name.
+	 * @return SessionName string.
+	 */
 	UFUNCTION(BlueprintGetter)
 	FString GetSessionName()
 	{
 		return SessionName;
 	}
 
+	/**
+	 * Set the session name
+	 * @param NewSessionName New session name string. This will be sanitized for illegal filesystem characters.
+	 */
 	UFUNCTION(BlueprintSetter)
 	void SetSessionName(FString NewSessionName)
 	{
@@ -92,12 +101,20 @@ public:
 		UpdateAllFields();
 	}
 
+	/**
+	 * Get the Production Name string.
+	 * @return ProductionName  
+	 */
 	UFUNCTION(BlueprintGetter)
 	FString GetProductionName()
 	{
 		return ProductionName;
 	}
 
+	/**
+	 * Set New Production Name string. This will be sanitized for illegal filesystem characters 
+	 * @param NewProductionName 
+	 */
 	UFUNCTION(BlueprintSetter)
 	void SetProductionName(FString NewProductionName)
 	{
@@ -140,53 +157,64 @@ public:
 	/**Take Recorder folder template*/
 	UPROPERTY(EditAnywhere, BlueprintReadOnly, Category="Folders|Session Subfolders", meta = (DisplayName = "Takes Folder", ContentDir, EditCondition="bIsEditable"))
 	FPCapTokenisedFolderPath TakeFolder;
-	
+
+	/** Additional folders. Create any additional folders you want under the session folder. You can use the Map key to label what each folder is for*/
 	UPROPERTY(EditAnywhere, BlueprintReadOnly, Category="Folders|Session Subfolders", meta=(DisplayAfter="Additional Folders", EditCondition="bIsEditable", ToolTip="Use the additional folders for any extra data you want contain in your session - eg audio and facial capture data"))
 	TMap<FName, FPCapTokenisedFolderPath> AdditionalFolders;
 	
-	/*
-	 * Take Record related fields
-	 */
+	// Take Record related fields */
 
+	/** Bool to control if timecode is recorded or not. */
 	UPROPERTY(EditAnywhere, BlueprintReadOnly, Category="Take Recorder", meta = (DisplayName = "Record Timecode", EditCondition="bIsEditable"))
 	bool bRecordTimecode = true;
 
+	/** Bool to control if all actors should be recorded to possessable or spawnable. */
 	UPROPERTY(EditAnywhere, BlueprintReadOnly, Category="Take Recorder", meta = (DisplayName = "Record to Possessable", EditCondition="bIsEditable"))
 	bool bRecordPossessable = false;
 
+	/** Bool to control if recorded sequences should start at the current timecode value. */
 	UPROPERTY(EditAnywhere, BlueprintReadOnly, Category="Take Recorder", meta = (DisplayName = "Start at Current Timecode", EditCondition="bIsEditable"))
 	bool bStartAtCurrentTimecode = true;
 
+	/** Bool to control if each recording source is placed into a subscene at the start of recording. */
 	UPROPERTY(EditAnywhere, BlueprintReadOnly, Category="Take Recorder", meta = (DisplayName = "Record to Subsequences", EditCondition="bIsEditable"))
 	bool bRecordSubscenes = true;
 
+	/** Take name token. */
 	UPROPERTY(EditAnywhere, BlueprintReadOnly, Category="Take Recorder", DisplayName = "Take Save Name", meta=(EditCondition="bIsEditable",  tooltip ="This be used by take recorder as the name of the recorded level sequence and the folder for any sub-scenes and assets"))
 	FPCapTokenisedString TakeSaveName;
 
+	/** Animation track name token. */
 	UPROPERTY(EditAnywhere, BlueprintReadOnly, Category="Take Recorder|Animation", DisplayName = "Animation Track Name", meta=(EditCondition="bIsEditable"))
 	FPCapTokenisedString AnimationTrackName;
 
+	/** Animation Asset name token. */
 	UPROPERTY(EditAnywhere, BlueprintReadOnly, Category="Take Recorder|Animation", DisplayName = "Animation Asset Name", meta=(EditCondition="bIsEditable"))
 	FPCapTokenisedString AnimationAssetName;
 
+	/**Animation subdirectory token. */
 	UPROPERTY(EditAnywhere, BlueprintReadOnly, Category="Take Recorder|Animation", DisplayName = "Animation Sub Directory", meta=(EditCondition="bIsEditable"))
 	FPCapTokenisedString AnimationSubDirectory;
 
+	/** Audio source name token. */
 	UPROPERTY(EditAnywhere, BlueprintReadOnly, Category="Take Recorder|Audio", DisplayName = "Audio Source Name", meta=(EditCondition="bIsEditable"))
 	FPCapTokenisedString AudioSourceName;
 
+	/** Audio track name token. */
 	UPROPERTY(EditAnywhere, BlueprintReadOnly, Category="Take Recorder|Audio", DisplayName = "Audio Track Name", meta=(EditCondition="bIsEditable"))
 	FPCapTokenisedString AudioTrackName;
 
+	/** Audio asset name token. */
 	UPROPERTY(EditAnywhere, BlueprintReadOnly, Category="Take Recorder|Audio", DisplayName = "Audio Asset Name", meta=(EditCondition="bIsEditable"))
 	FPCapTokenisedString AudioAssetName;
 
+	/** Audio subdirectory token. */
 	UPROPERTY(EditAnywhere, BlueprintReadOnly, Category="Take Recorder|Audio", DisplayName = "Audio Sub Directory", meta=(EditCondition="bIsEditable"))
 	FPCapTokenisedString AudioSubDirectory;
+
 	/*
 	 * Read only fields to show the user what tokens are available to them in the pcap, global, and take recorder namespaces
 	 */
-	
 	UPROPERTY(VisibleAnywhere, Category="Tokens", meta= (displayname = "PCap Tokens", tooltip="Use these  tokens in your template strings. The namespace for these tokens is 'pcap.'")) // this is for display purposes only, not to be used in code.
 	FString PCapTokens =
 	"{session} - Name for your Performance Capture session\n"
@@ -240,6 +268,7 @@ protected:
 	 * Behind the scenes, this prevents PostEditUpdates from modifying data such that makes all
 	 * the strings and folder paths become static
 	 */
+	/** Bool to control edit condition on the members of this asset. Prevents user from editing after a session has been created.*/
 	UPROPERTY(BlueprintReadWrite, VisibleAnywhere, Category="Admin")
 	bool bIsEditable;
 	
@@ -247,11 +276,24 @@ protected:
 	virtual void PostEditChangeProperty(FPropertyChangedEvent& PropertyChangedEvent) override;
 #endif
 
+	/**
+	 * Update all members and ensure illegal characters are removed. 
+	 */
 	UFUNCTION(BlueprintCallable, Category="Performance Capture|Session Template")
 	void UpdateAllFields();
 
+	/**
+	 * Update string token from the given template.
+	 * @param TokenisedTemplate 
+	 * @return Updated string token struct.
+	 */
 	FPCapTokenisedString UpdateStringTemplate(FPCapTokenisedString& TokenisedTemplate);
 
+	/**
+	 * Upate a token for a folder path.
+	 * @param FolderPathTokenisedTemplate 
+	 * @return Updated folder path struct.
+	 */
 	FPCapTokenisedFolderPath UpdateFolderPathTemplate(FPCapTokenisedFolderPath& FolderPathTokenisedTemplate);
 	
 private:
diff --git a/Engine/Plugins/VirtualProduction/PerformanceCaptureWorkflow/Source/PerformanceCaptureWorkflow/Private/PCapSettings.h b/Engine/Plugins/VirtualProduction/PerformanceCaptureWorkflow/Source/PerformanceCaptureWorkflow/Private/PCapSettings.h
--- a/Engine/Plugins/VirtualProduction/PerformanceCaptureWorkflow/Source/PerformanceCaptureWorkflow/Private/PCapSettings.h
+++ b/Engine/Plugins/VirtualProduction/PerformanceCaptureWorkflow/Source/PerformanceCaptureWorkflow/Private/PCapSettings.h
@@ -25,53 +25,73 @@ class PERFORMANCECAPTUREWORKFLOW_API UPerformanceCaptureSettings : public UObjec
 {
 	GENERATED_BODY()
 public:
-	
+	/**
+	 * Get the Performance Capture settings uobject
+	 * @return Settings Object.
+	 */
 	UFUNCTION(BlueprintPure, Category="Performance Capture|Settings")
 	static UPerformanceCaptureSettings* GetPerformanceCaptureSettings();
 
+	/**
+	 * Open the project settings panel and show the Performance Capture section.
+	 */
 	UFUNCTION(BlueprintCallable, Category = "Performance Capture|Settings")
 	void ShowPerformanceCaptureProjectSettings();
 	
-	/**Stage Root Actor to spawn for Performance Capture. Class must derive from APerformanceCaptureStageRoot*/
+	/** Stage Root Actor to spawn for Performance Capture. Class must derive from APerformanceCaptureStageRoot. */
 	UPROPERTY(BlueprintReadOnly, EditDefaultsOnly, Config, NoClear, Category="Actors", DisplayName= "Stage Root Actor Class" ,meta = (MetaClass = "/Script/PerformanceCaptureWorkflow.PerformanceCaptureStageRoot"))
 	TSoftClassPtr<APerformanceCaptureStageRoot> StageRoot;
 
-	///
-	/**Base skeletal mesh for Mocap performers*/
+	
+	/** Base skeletal mesh for Mocap performers. */
 	UPROPERTY(BlueprintReadOnly, EditDefaultsOnly, Config, NoClear, Category="Actors", DisplayName= "Default Performer Mesh")
 	TSoftObjectPtr<USkeletalMesh> DefaultPerformerSkelMesh;
 
-	/** Blueprint Viewmodel class. If you change class you will need to restart the editor to pick up the change*/
+	/** Blueprint Viewmodel class. If you change class you will need to restart the editor to pick up the change. */
 	UPROPERTY(EditAnywhere, BlueprintReadOnly, Config, NoClear, Category="UI", meta = (MetaClass = "/Script/ModelViewViewModel.MVVMViewModelBase",ConfigRestartRequired=true))
 	TSoftClassPtr<UMVVMViewModelBase> ViewModelClass;
 
 	/** Editor Utility Widget class that will be used for the Mocap Manager UI*/
 	UPROPERTY(EditAnywhere, BlueprintReadOnly, Config, Category="UI", meta = (MetaClass = "/Script/Blutility.EditorUtilityWidgetBlueprint"))
 	TSoftObjectPtr<UEditorUtilityWidgetBlueprint> MocapManagerUI;
-	
+
+	/** Default session template to use when starting the Mocap Manager panel. */
 	UPROPERTY(BlueprintReadOnly, Config, EditDefaultsOnly, Category="Database", DisplayName = "Default Session Template", meta = ( tooltip="This Asset defines the folder structure and settings used in a performance capture session"))
 	TSoftObjectPtr<UPCapSessionTemplate> DefaultSessionTemplate;
 	
 	/** Blueprint helper class for making database calls. If you change class you will need to restart the editor to pick up the change*/
 	UPROPERTY(EditAnywhere, Config, NoClear, Category="Database", DisplayName = "Database Helper Class", meta=(ConfigRestartRequired=true))
 	TSoftClassPtr<UPerformanceCaptureDatabaseHelper> DatabaseHelperClass;
-	
+
+	/** Pointer to the datatable Mocap Manager will use for recording all session data. */
 	UPROPERTY(BlueprintReadOnly, Config, EditDefaultsOnly, Category = "Database|Internal Database Tables", DisplayName = "Session DataTable", meta = (RequiredAssetDataTags = "RowStructure=/Script/PerformanceCaptureWorkflow.PCapSessionRecord"))
 	TSoftObjectPtr<UPCapDataTable> SessionTable;
 
+	/** Proint to the datatable Mocap Manager will use for recording production data. */
 	UPROPERTY(BlueprintReadOnly, Config, EditDefaultsOnly, Category="Database|Internal Database Tables", DisplayName = "Production DataTable", meta = (RequiredAssetDataTags = "RowStructure=/Script/PerformanceCaptureWorkflow.PCapProductionRecord"))
 	TSoftObjectPtr<UPCapDataTable> ProductionTable;
 
+	/**
+	 * Set the session table in Performance Capture settings. 
+	 * @param NewDataTable Datatable ref. Must be of the FPCapSession struct type.
+	 */
 	UFUNCTION(BlueprintCallable, Category="Performance Capture|Settings")
 	void SetSessionTable(TSoftObjectPtr<UPCapDataTable> NewDataTable);
 
+	/**
+	 * Set the production table in Performance Capture Settings.
+	 * @param NewDataTable Datatable ref. Must be of the FPCapProduction struct type.
+	 */
 	UFUNCTION(BlueprintCallable, Category="Performance Capture|Settings")
 	void SetProductionTable(TSoftObjectPtr<UPCapDataTable> NewDataTable); 
 
 #if WITH_EDITOR
 	
 	DECLARE_DYNAMIC_MULTICAST_DELEGATE(FOnPCapSettingsChanged);
-	
+
+	/**
+	 * Multicast delegate called whenever the Performance Capture settings object is modified.
+	 */
 	UPROPERTY(BlueprintAssignable, Category="Performance Capture|Settings")
 	FOnPCapSettingsChanged OnPCapSettingsChanged;
 	
diff --git a/Engine/Plugins/VirtualProduction/PerformanceCaptureWorkflow/Source/PerformanceCaptureWorkflow/Private/PCapSubsystem.h b/Engine/Plugins/VirtualProduction/PerformanceCaptureWorkflow/Source/PerformanceCaptureWorkflow/Private/PCapSubsystem.h
--- a/Engine/Plugins/VirtualProduction/PerformanceCaptureWorkflow/Source/PerformanceCaptureWorkflow/Private/PCapSubsystem.h
+++ b/Engine/Plugins/VirtualProduction/PerformanceCaptureWorkflow/Source/PerformanceCaptureWorkflow/Private/PCapSubsystem.h
@@ -33,12 +33,20 @@ public:
 	virtual void Initialize(FSubsystemCollectionBase& Collection) override;
 	// ~End UEditorSubsystem
 
+	/**
+	 * Get the Performance Capture Database helper.
+	 * @return Databasehelper object
+	 */
 	UFUNCTION(BlueprintCallable, Category="Performance Capture|Database")
 	UPerformanceCaptureDatabaseHelper* GetDatabaseHelper() const
 	{
 		return DatabaseHelper;
 	}
 
+	/**
+	 * Get the Performance Capture Viewmodel collection.
+	 * @return Performance Capture Viewmodel collection.
+	 */
 	UFUNCTION(BlueprintCallable, Category="Performance Capture|Viewmodel")
 	UMVVMViewModelCollectionObject* GetViewModelCollection() const
 	{
diff --git a/Engine/Plugins/VirtualProduction/PerformanceCaptureWorkflow/Source/PerformanceCaptureWorkflow/Private/Visualizers/LiveLinkDataPreview.cpp b/Engine/Plugins/VirtualProduction/PerformanceCaptureWorkflow/Source/PerformanceCaptureWorkflow/Private/Visualizers/LiveLinkDataPreview.cpp
--- a/Engine/Plugins/VirtualProduction/PerformanceCaptureWorkflow/Source/PerformanceCaptureWorkflow/Private/Visualizers/LiveLinkDataPreview.cpp
+++ b/Engine/Plugins/VirtualProduction/PerformanceCaptureWorkflow/Source/PerformanceCaptureWorkflow/Private/Visualizers/LiveLinkDataPreview.cpp
@@ -58,6 +58,7 @@ void ALiveLinkDataPreview::Tick(float DeltaTime)
 	
 		TArray<FLiveLinkSubjectKey> SubjectKeys = LiveLinkClient.GetSubjects(false, false);
 
+		//If the number of Live Link subjects changes, force reinitialization. 
 		if(SubjectKeys.Num() != CachedSubjects.Num())
 		{
 			InitializeSubjects();
@@ -92,10 +93,13 @@ void ALiveLinkDataPreview::PostEditChangeProperty(FPropertyChangedEvent& Propert
 	}
 }
 
+
 void ALiveLinkDataPreview::SetEnableLiveLinkData(bool bNewEvaluate)
 {
 	bEvaluateLiveLink = bNewEvaluate;
 
+	//Loop through each subject's preview component and set it's bEvaluateLiveLink member.
+
 	TArray<ULiveLinkDataPreviewComponent*> DebugComponents;
 	GetComponents(DebugComponents);
 
@@ -105,6 +109,7 @@ void ALiveLinkDataPreview::SetEnableLiveLinkData(bool bNewEvaluate)
 	}
 }
 
+//Initialize each subject's preview component and depending on the role type, add a component with the mesh of the correct type.
 void ALiveLinkDataPreview::InitializeSubjects()
 {
 	if((GetWorld() && !GetWorld()->IsPreviewWorld())) //Do not spawn the preview components unless we are a non-preview world
@@ -183,6 +188,7 @@ void ALiveLinkDataPreview::InitializeSubjects()
 	}
 }
 
+//Set the material instance and set a nice random color. 
 void ALiveLinkDataPreview::SetMaterialInstance(ULiveLinkDataPreviewComponent* InDataPreviewComponent)
 {
 	UMaterialInterface* DebugMaterial = InDataPreviewComponent->GetMaterial(0);
diff --git a/Engine/Plugins/VirtualProduction/PerformanceCaptureWorkflow/Source/PerformanceCaptureWorkflow/Private/Visualizers/LiveLinkDataPreview.h b/Engine/Plugins/VirtualProduction/PerformanceCaptureWorkflow/Source/PerformanceCaptureWorkflow/Private/Visualizers/LiveLinkDataPreview.h
--- a/Engine/Plugins/VirtualProduction/PerformanceCaptureWorkflow/Source/PerformanceCaptureWorkflow/Private/Visualizers/LiveLinkDataPreview.h
+++ b/Engine/Plugins/VirtualProduction/PerformanceCaptureWorkflow/Source/PerformanceCaptureWorkflow/Private/Visualizers/LiveLinkDataPreview.h
@@ -22,6 +22,7 @@ public:
 	// Sets default values for this actor's properties
 	ALiveLinkDataPreview();
 
+	/** The list of Live Link subjects this actor will draw. */
 	UPROPERTY(EditAnywhere, BlueprintReadWrite, Category = LiveLink, meta=(ShowOnlyInnerProperties="true"))
 	TArray<FLiveLinkSubjectName> Subjects;
 
@@ -35,22 +36,34 @@ public:
 	UPROPERTY(EditAnywhere, BlueprintReadWrite, Category = LiveLink)
 	bool bDrawLabels;
 
-	UPROPERTY(EditAnywhere, BlueprintReadWrite, Category = LiveLink)
+	// Enable Live Link evaluation
+	UPROPERTY(EditAnywhere, Category = LiveLink)
 	bool bEvaluateLiveLink;
 
 	TArray<ULiveLinkDataPreviewComponent> Visualizers;
-	
+
+	/**
+	 * Stop Start animation evaluation.
+	 * @param bNewEvaluate Enable Live Link Evaluation
+	 */
 	UFUNCTION(CallInEditor, BlueprintCallable, DisplayName="Evaluate Live Link", Category = LiveLink)
 	void SetEnableLiveLinkData(bool bNewEvaluate);
 
+	/** The billboard component showing the actor's root position. */
 	UPROPERTY(EditAnywhere, BlueprintReadWrite, Category = Components)
 	TObjectPtr<UBillboardComponent> BillboardComponent;
 
 	protected:
-
+	/**
+	 * Initialize all the subjects.
+	 */
 	UFUNCTION(BlueprintCallable, DisplayName="Initialize All Subjects", Category = LiveLink)
 	void InitializeSubjects();
 
+	/**
+	 * Set the material instance on a Live Link data preview component.
+	 * @param InDataPreviewComponent Preview component instanced static mesh.
+	 */
 	void SetMaterialInstance(ULiveLinkDataPreviewComponent* InDataPreviewComponent);
 	
 	// Called every frame
diff --git a/Engine/Plugins/VirtualProduction/PerformanceCaptureWorkflow/Source/PerformanceCaptureWorkflow/Private/Visualizers/LiveLinkDataPreviewComponent.cpp b/Engine/Plugins/VirtualProduction/PerformanceCaptureWorkflow/Source/PerformanceCaptureWorkflow/Private/Visualizers/LiveLinkDataPreviewComponent.cpp
--- a/Engine/Plugins/VirtualProduction/PerformanceCaptureWorkflow/Source/PerformanceCaptureWorkflow/Private/Visualizers/LiveLinkDataPreviewComponent.cpp
+++ b/Engine/Plugins/VirtualProduction/PerformanceCaptureWorkflow/Source/PerformanceCaptureWorkflow/Private/Visualizers/LiveLinkDataPreviewComponent.cpp
@@ -27,6 +27,7 @@ ULiveLinkDataPreviewComponent::ULiveLinkDataPreviewComponent(const FObjectInitia
 void ULiveLinkDataPreviewComponent::TickComponent(float DeltaTime, ELevelTick TickType, FActorComponentTickFunction* ThisTickFunction)
 {
 	Super::TickComponent(DeltaTime, TickType, ThisTickFunction);
+	//TODO: add drawing labels for bones and transforms. 
 	
 	ILiveLinkClient& LiveLinkClient = IModularFeatures::Get().GetModularFeature<ILiveLinkClient>(ILiveLinkClient::ModularFeatureName);
 
@@ -34,7 +35,7 @@ void ULiveLinkDataPreviewComponent::TickComponent(float DeltaTime, ELevelTick Ti
 	
 	if (SubjectRole && !bIsDirty)
 	{
-		//Process Animation Data
+		//Process Animation Data. Animation Role handling.
 		if(LiveLinkClient.DoesSubjectSupportsRole_AnyThread(SubjectName, ULiveLinkAnimationRole::StaticClass()))
 		{
 			if(LiveLinkClient.IsSubjectValid(SubjectName)) //Check there is valid data on the subject before evaluate it. 
@@ -72,6 +73,7 @@ void ULiveLinkDataPreviewComponent::TickComponent(float DeltaTime, ELevelTick Ti
 				}
 			}
 		}
+		//Locator role handling
 		else if(LiveLinkClient.DoesSubjectSupportsRole_AnyThread(SubjectName, ULiveLinkLocatorRole::StaticClass()) && bEvaluateLiveLink)
 		{
 			
@@ -87,6 +89,7 @@ void ULiveLinkDataPreviewComponent::TickComponent(float DeltaTime, ELevelTick Ti
 			}
 			BatchUpdateInstancesTransforms(0, Transforms, true, true, false);
 		}
+		//Transform role handling.
 		else if(LiveLinkClient.DoesSubjectSupportsRole_AnyThread(SubjectName, ULiveLinkTransformRole::StaticClass()) && bEvaluateLiveLink)
 		{
 			BatchUpdateInstancesTransform(0,1,GetSingleTransform() * GetComponentTransform(),true,true,false);
@@ -195,6 +198,7 @@ void ULiveLinkDataPreviewComponent::GetTransformRootSpace(const int32 InTransfor
 	}
 }
 
+//Add the static mesh instances to represent the Live Link data. 
 void ULiveLinkDataPreviewComponent::CreateInstances()
 {
 	ILiveLinkClient& LiveLinkClient = IModularFeatures::Get().GetModularFeature<ILiveLinkClient>(ILiveLinkClient::ModularFeatureName);
@@ -245,6 +249,7 @@ void ULiveLinkDataPreviewComponent::CreateInstances()
 	}
 }
 
+//Loop through a bone's parent's to get it's transform in actor space. Live Link bone transforms are expressed in parent bone space so we need to do this to convert to actor space. 
 TArray<FTransform> ULiveLinkDataPreviewComponent::GetBoneTransforms() const
 {
 	TArray<FTransform> JointTransforms;
@@ -274,6 +279,7 @@ TArray<FTransform> ULiveLinkDataPreviewComponent::GetBoneTransforms() const
 	return JointTransforms;
 }
 
+//Get bone transforms from the cached skeleton data.
 TArray<FTransform> ULiveLinkDataPreviewComponent::GetJointTransforms() const
 {
 	TArray<FTransform> BoneTransforms;
@@ -291,6 +297,7 @@ TArray<FTransform> ULiveLinkDataPreviewComponent::GetJointTransforms() const
 	return BoneTransforms;
 }
 
+//Get the location of each locator in the locator role and apply it to the instanced static meshes.
 TArray<FVector> ULiveLinkDataPreviewComponent::GetLocatorPositions() const
 {
 	ILiveLinkClient& LiveLinkClient = IModularFeatures::Get().GetModularFeature<ILiveLinkClient>(ILiveLinkClient::ModularFeatureName);
@@ -324,6 +331,7 @@ TArray<FVector> ULiveLinkDataPreviewComponent::GetLocatorPositions() const
 	return MarkerPostions;
 }
 
+//Get a tranform role's data and applying it to the instance static mesh in actor space.
 FTransform ULiveLinkDataPreviewComponent::GetSingleTransform() const
 {
 	ILiveLinkClient& LiveLinkClient = IModularFeatures::Get().GetModularFeature<ILiveLinkClient>(ILiveLinkClient::ModularFeatureName);
@@ -345,6 +353,7 @@ FTransform ULiveLinkDataPreviewComponent::GetSingleTransform() const
 	return Transform;
 }
 
+//Cached a frame's Live Link Animation role bone data for use later in the frame, updating in the instance static mesh components. 
 void ULiveLinkDataPreviewComponent::CacheSkeletalAnimationData(const FLiveLinkSkeletonStaticData* InStaticData, const FLiveLinkAnimationFrameData* InFrameData)
 {
 	{
@@ -365,6 +374,7 @@ void ULiveLinkDataPreviewComponent::CacheSkeletalAnimationData(const FLiveLinkSk
 	}
 }
 
+//Check a transform is valid, based on the given index. 
 bool ULiveLinkDataPreviewComponent::IsValidTransformIndex(int32 InTransformIndex, FLiveLinkAnimationFrameData InAnimData)
 {
 	return (InTransformIndex >= 0) && (InTransformIndex < InAnimData.Transforms.Num());
diff --git a/Engine/Plugins/VirtualProduction/PerformanceCaptureWorkflow/Source/PerformanceCaptureWorkflow/Private/Visualizers/LiveLinkDataPreviewComponent.h b/Engine/Plugins/VirtualProduction/PerformanceCaptureWorkflow/Source/PerformanceCaptureWorkflow/Private/Visualizers/LiveLinkDataPreviewComponent.h
--- a/Engine/Plugins/VirtualProduction/PerformanceCaptureWorkflow/Source/PerformanceCaptureWorkflow/Private/Visualizers/LiveLinkDataPreviewComponent.h
+++ b/Engine/Plugins/VirtualProduction/PerformanceCaptureWorkflow/Source/PerformanceCaptureWorkflow/Private/Visualizers/LiveLinkDataPreviewComponent.h
@@ -32,27 +32,40 @@ class PERFORMANCECAPTUREWORKFLOW_API ULiveLinkDataPreviewComponent : public UIns
 
 	public:
 
-	//Sets default values for this component's properties
+	// Sets default values for this component's properties
 	ULiveLinkDataPreviewComponent(const FObjectInitializer& ObjectInitializer);
 
-	//Called every frame
+	// Called every frame
 	virtual void TickComponent(float DeltaTime, ELevelTick TickType, FActorComponentTickFunction* ThisTickFunction) override;
 
+	/** The Live Link subject this component will preview. */
 	UPROPERTY(EditAnywhere, BlueprintReadWrite, Category = LiveLink)
 	FLiveLinkSubjectName SubjectName;
 
+	/** Bool to control animation evaluation. */
 	UPROPERTY(EditAnywhere, BlueprintReadOnly, Category=LiveLink)
 	bool bEvaluateLiveLink;
 
+	/** Bool to control visibility of labels. */
 	UPROPERTY(EditAnywhere, BlueprintReadOnly, Category="PerformanceCapture|Visualization")
 	bool bDrawLabels;
 
+	/** Type of bone visualization - joint or bone.  */
 	UPROPERTY(BlueprintReadWrite, EditAnywhere, Category="PerformanceCapture|Visualization")
 	ELiveLinkVisualBoneType BoneVisualType = ELiveLinkVisualBoneType::Bone;
 
+	/**
+	 * Stop/Start animation updates. 
+	 * @param bNewEvaluateLiveLink New evaluation state.
+	 */
 	UFUNCTION(BlueprintCallable, Category="PerformanceCapture|Visualization")
 	void SetEvaluateLiveLinkData(bool bNewEvaluateLiveLink) { bEvaluateLiveLink = bNewEvaluateLiveLink; }
 
+	/**
+	 * Show or Hide Labels
+	 * @param bNewDrawLabel New lable visiblity..
+	 */
+	//TODO: not yet implemented, waiting on slate immediate mode
 	UFUNCTION(BlueprintCallable, Category="PerformanceCapture|Visualization")
 	void SetDrawLabels(bool bNewDrawLabel) { bDrawLabels = bNewDrawLabel; }
 	
diff --git a/Engine/Plugins/VirtualProduction/PerformanceCaptureWorkflow/Source/PerformanceCaptureWorkflow/Private/Visualizers/PCapBoneVisualizer.h b/Engine/Plugins/VirtualProduction/PerformanceCaptureWorkflow/Source/PerformanceCaptureWorkflow/Private/Visualizers/PCapBoneVisualizer.h
--- a/Engine/Plugins/VirtualProduction/PerformanceCaptureWorkflow/Source/PerformanceCaptureWorkflow/Private/Visualizers/PCapBoneVisualizer.h
+++ b/Engine/Plugins/VirtualProduction/PerformanceCaptureWorkflow/Source/PerformanceCaptureWorkflow/Private/Visualizers/PCapBoneVisualizer.h
@@ -29,23 +29,29 @@ public:
 	// Called every frame
 	virtual void TickComponent(float DeltaTime, ELevelTick TickType, FActorComponentTickFunction* ThisTickFunction) override;
 
-	/** Get all the transforms from the skelmesh*/
+	/** Get all the transforms from the skelmesh. */
 	TArray<FTransform> GetJointTransforms(USkeletalMeshComponent* InSkelMeshComponent) const;
 
-	/** Get all the bones' transforms for drawing a joint (point to point mesh)*/
+	/** Get all the bones' transforms for drawing a joint (point to point mesh). */
 	TArray<FTransform> GetBoneTransforms(USkeletalMeshComponent* InSkelMeshComponent) const;
 	
-	/** Whether to draw joints along the length of the bone, or just at the pivot or scaled by the distance to the parent bone*/
+	/** Whether to draw joints along the length of the bone, or just at the pivot or scaled by the distance to the parent bone. */
 	UPROPERTY(EditAnywhere, BlueprintReadWrite, Category="PerformanceCapture|Visualization")
 	EBoneVisualType VisualizationType;
 
+	/** Color to use on the instanced static meshes. Can only be set during construction. */
 	UPROPERTY(EditAnywhere, BlueprintReadWrite, Setter="UpdateColor", Category="PerformanceCapture|Visualization")
 	FLinearColor Color;
-	
+
+	/** Dynamic Material to use on the instanced static meshes. */
 	UPROPERTY(EditAnywhere, BlueprintReadOnly, Category="PerformanceCapture|Visualization")
 	TObjectPtr<UMaterialInstanceDynamic> DynamicMaterial;
 
-	UFUNCTION()
+	/**
+	 * Update the color for this bone visaulizer.
+	 * @param NewColor New color to apply in the material.
+	 */
+	UFUNCTION(BlueprintCallable, Category="PerformanceCapture|Visualization")
 	void UpdateColor(FLinearColor NewColor);
 
 #if WITH_EDITOR
diff --git a/Engine/Plugins/VirtualProduction/PerformanceCaptureWorkflow/Source/PerformanceCaptureWorkflow/Private/Visualizers/PCapMarkerVisualizer.cpp b/Engine/Plugins/VirtualProduction/PerformanceCaptureWorkflow/Source/PerformanceCaptureWorkflow/Private/Visualizers/PCapMarkerVisualizer.cpp
--- a/Engine/Plugins/VirtualProduction/PerformanceCaptureWorkflow/Source/PerformanceCaptureWorkflow/Private/Visualizers/PCapMarkerVisualizer.cpp
+++ b/Engine/Plugins/VirtualProduction/PerformanceCaptureWorkflow/Source/PerformanceCaptureWorkflow/Private/Visualizers/PCapMarkerVisualizer.cpp
@@ -31,11 +31,12 @@ void UPCapMarkerVisualizer::TickComponent(float DeltaTime, ELevelTick TickType,
 {
 	Super::TickComponent(DeltaTime, TickType, ThisTickFunction);
 
+	//If there are no instances, add them. 
 	if(GetInstanceCount() == 0)
 	{
 		AddInstances(GetMarkerTransforms(), true,true,false);
 	}
-
+	
 	if(bEvaluateLiveLink)
 	{
 		BatchUpdateInstancesTransforms(0, GetMarkerTransforms(), true,true,false);
diff --git a/Engine/Plugins/VirtualProduction/PerformanceCaptureWorkflow/Source/PerformanceCaptureWorkflow/Private/Visualizers/PCapMarkerVisualizer.h b/Engine/Plugins/VirtualProduction/PerformanceCaptureWorkflow/Source/PerformanceCaptureWorkflow/Private/Visualizers/PCapMarkerVisualizer.h
--- a/Engine/Plugins/VirtualProduction/PerformanceCaptureWorkflow/Source/PerformanceCaptureWorkflow/Private/Visualizers/PCapMarkerVisualizer.h
+++ b/Engine/Plugins/VirtualProduction/PerformanceCaptureWorkflow/Source/PerformanceCaptureWorkflow/Private/Visualizers/PCapMarkerVisualizer.h
@@ -10,6 +10,9 @@
 
 class UMaterialInstanceDynamic;
 
+/**
+ * Sign of the axis, to handle data coming from non-left handed systems.
+ */
 UENUM(BlueprintType)
 enum class EAxisSign : uint8
 {
@@ -40,29 +43,31 @@ public:
 
 	static void DrawLabels(TArray<FName> Labels);
 
-	/*UPROPERTY(EditAnywhere, BlueprintReadWrite, Category="PerformanceCapture|Visualization")
-	EMarkerVisualType VisualizationType;*/
-	
+	/** Dynamic material to use on the instance static meshes. */
 	UPROPERTY(EditAnywhere, BlueprintReadOnly, Category="PerformanceCapture|Visualization")
 	TObjectPtr<UMaterialInstanceDynamic> DynamicMaterial;
 
+	/** Live Link subject to drive the instanced static mesh transforms. */
 	UPROPERTY(BlueprintReadWrite, EditAnywhere, Category="PerformanceCapture|Visualization")
 	FLiveLinkSubjectName LiveLinkSubject;
 
-	//To Do - add text labels to drawn markers
+	/** Bool to control evaluation of Live Link data. */
 	UPROPERTY(BlueprintReadWrite, EditAnywhere, Category="PerformanceCapture|Visualization")
 	bool bEvaluateLiveLink;
 
 	//To Do - add text labels to drawn markers
 	UPROPERTY(BlueprintReadWrite, EditAnywhere, Category="PerformanceCapture|Visualization")
 	bool bDrawLabels;
-	
+
+	/** Sign for the x axis of the marker. */
 	UPROPERTY(BlueprintReadWrite, EditAnywhere, Category="PerformanceCapture|Visualization")
 	EAxisSign XAxisSign;
 
+	/** Sign for the y axis of the marker. */
 	UPROPERTY(BlueprintReadWrite, EditAnywhere, Category="PerformanceCapture|Visualization")
 	EAxisSign YAxisSign;
-	
+
+	/** Array of Transforms for the markers.  */
 	UPROPERTY(BlueprintReadWrite, EditAnywhere, Category="PerformanceCapture|Visualization")
 	TArray<FTransform> MarkerLocations;
 
diff --git a/Engine/Plugins/VirtualProduction/PerformanceCaptureWorkflow/Source/PerformanceCaptureWorkflow/Public/PCapStageRoot.h b/Engine/Plugins/VirtualProduction/PerformanceCaptureWorkflow/Source/PerformanceCaptureWorkflow/Public/PCapStageRoot.h
--- a/Engine/Plugins/VirtualProduction/PerformanceCaptureWorkflow/Source/PerformanceCaptureWorkflow/Public/PCapStageRoot.h
+++ b/Engine/Plugins/VirtualProduction/PerformanceCaptureWorkflow/Source/PerformanceCaptureWorkflow/Public/PCapStageRoot.h
@@ -21,13 +21,16 @@ class APerformanceCaptureStageRoot : public AActor
 public:
 	// Sets default values for this actor's properties
 	APerformanceCaptureStageRoot(const FObjectInitializer& ObjectInitializer);
-	
+
+	/** Scene capture component for capturing the overhead, orthographic map. */
 	UPROPERTY(EditAnywhere, BlueprintReadWrite,  Category="Peformance Capture|StageRoot")
 	TObjectPtr<USceneCaptureComponent2D> MapCaptureComponent;
 
+	/** Decal component for rendering a grid to the floor. */
 	UPROPERTY(EditAnywhere, BlueprintReadWrite,  Category="Peformance Capture|StageRoot")
 	TObjectPtr<UDecalComponent> DecalComponent;
 
+	/** Scene component, under which all stage ghost meshes should be parented. */
 	UPROPERTY(EditAnywhere, BlueprintReadWrite,  Category="Peformance Capture|StageRoot")
 	TObjectPtr<USceneComponent> StageMeshParent;
 	
diff --git a/Engine/Plugins/VirtualProduction/PerformanceCaptureWorkflow/Source/PerformanceCaptureWorkflow/Public/PCapViewmodel.h b/Engine/Plugins/VirtualProduction/PerformanceCaptureWorkflow/Source/PerformanceCaptureWorkflow/Public/PCapViewmodel.h
--- a/Engine/Plugins/VirtualProduction/PerformanceCaptureWorkflow/Source/PerformanceCaptureWorkflow/Public/PCapViewmodel.h
+++ b/Engine/Plugins/VirtualProduction/PerformanceCaptureWorkflow/Source/PerformanceCaptureWorkflow/Public/PCapViewmodel.h
@@ -15,6 +15,7 @@ class PERFORMANCECAPTUREWORKFLOW_API UPCapViewmodel : public UMVVMViewModelBase
 	GENERATED_BODY()
 public:
 
+	/** Initialise the viewmodel. Implemented in Blueprint. */
 	UFUNCTION(BlueprintCallable, BlueprintImplementableEvent, Category="Performance Capture|Viewmodel")
 	void Initialize();
 };
