[PCap Workflow]
- following late change to take recorder, updates the session template to  support the new Subsequence Directory property.

#rnx

#virtualized

diff --git a/Engine/Plugins/VirtualProduction/PerformanceCaptureWorkflow/Source/PerformanceCaptureWorkflow/Private/PCapSessionTemplate.cpp b/Engine/Plugins/VirtualProduction/PerformanceCaptureWorkflow/Source/PerformanceCaptureWorkflow/Private/PCapSessionTemplate.cpp
--- a/Engine/Plugins/VirtualProduction/PerformanceCaptureWorkflow/Source/PerformanceCaptureWorkflow/Private/PCapSessionTemplate.cpp
+++ b/Engine/Plugins/VirtualProduction/PerformanceCaptureWorkflow/Source/PerformanceCaptureWorkflow/Private/PCapSessionTemplate.cpp
@@ -10,7 +10,7 @@ UPCapSessionTemplate::UPCapSessionTemplate(const FObjectInitializer& ObjectIniti
 {
 	bIsEditable = true;
 	NamingTokensContext = NewObject<UPCapNamingTokensContext>();
-	
+	SubsequenceDirectory.Template = FString("{takeName}_Subscenes"); //Matched to the default in take recorder.
 }
 
 void UPCapSessionTemplate::PostLoad()
@@ -55,6 +55,7 @@ void UPCapSessionTemplate::UpdateAllFields()
 		AnimationTrackName	= UpdateStringTemplate(AnimationTrackName);
 		AnimationAssetName	= UpdateStringTemplate(AnimationAssetName);
 		AnimationSubDirectory= UpdateStringTemplate(AnimationSubDirectory);
+		SubsequenceDirectory= UpdateStringTemplate(SubsequenceDirectory);
 		AudioSourceName		= UpdateStringTemplate(AudioSourceName);
 		AudioTrackName		= UpdateStringTemplate(AudioTrackName);
 		AudioAssetName		= UpdateStringTemplate(AudioAssetName);
diff --git a/Engine/Plugins/VirtualProduction/PerformanceCaptureWorkflow/Source/PerformanceCaptureWorkflow/Private/PCapSessionTemplate.h b/Engine/Plugins/VirtualProduction/PerformanceCaptureWorkflow/Source/PerformanceCaptureWorkflow/Private/PCapSessionTemplate.h
--- a/Engine/Plugins/VirtualProduction/PerformanceCaptureWorkflow/Source/PerformanceCaptureWorkflow/Private/PCapSessionTemplate.h
+++ b/Engine/Plugins/VirtualProduction/PerformanceCaptureWorkflow/Source/PerformanceCaptureWorkflow/Private/PCapSessionTemplate.h
@@ -187,7 +187,7 @@ public:
 	/** Animation track name token. */
 	UPROPERTY(EditAnywhere, BlueprintReadOnly, Category="Take Recorder|Animation", DisplayName = "Animation Track Name", meta=(EditCondition="bIsEditable"))
 	FPCapTokenisedString AnimationTrackName;
-
+	
 	/** Animation Asset name token. */
 	UPROPERTY(EditAnywhere, BlueprintReadOnly, Category="Take Recorder|Animation", DisplayName = "Animation Asset Name", meta=(EditCondition="bIsEditable"))
 	FPCapTokenisedString AnimationAssetName;
@@ -196,6 +196,10 @@ public:
 	UPROPERTY(EditAnywhere, BlueprintReadOnly, Category="Take Recorder|Animation", DisplayName = "Animation Sub Directory", meta=(EditCondition="bIsEditable"))
 	FPCapTokenisedString AnimationSubDirectory;
 
+	/** Subsequence directory name token. */
+	UPROPERTY(EditAnywhere, BlueprintReadOnly, Category="Take Recorder|Animation", DisplayName = "Subsequence Directory", meta=(EditCondition="bIsEditable", tooltip ="This affects the directory created for subsequences. It cannot be set per actor, only globally across the whole recording."))
+	FPCapTokenisedString SubsequenceDirectory;
+	
 	/** Audio source name token. */
 	UPROPERTY(EditAnywhere, BlueprintReadOnly, Category="Take Recorder|Audio", DisplayName = "Audio Source Name", meta=(EditCondition="bIsEditable"))
 	FPCapTokenisedString AudioSourceName;
