LiveLink / PCap - Move Axis, Bone and Joint visualizer back to their original plugin

#virtualized

diff --git a/Engine/Plugins/Animation/LiveLink/Source/LiveLink/Private/Visualizers/LiveLinkBoneVisualizer.cpp b/Engine/Plugins/Animation/LiveLink/Source/LiveLink/Private/Visualizers/LiveLinkBoneVisualizer.cpp
deleted file mode 100644
diff --git a/Engine/Plugins/Animation/LiveLink/Source/LiveLink/Private/Visualizers/LiveLinkBoneVisualizer.h b/Engine/Plugins/Animation/LiveLink/Source/LiveLink/Private/Visualizers/LiveLinkBoneVisualizer.h
deleted file mode 100644
diff --git a/Engine/Plugins/Animation/LiveLink/Source/LiveLink/Private/Visualizers/LiveLinkDataPreviewComponent.cpp b/Engine/Plugins/Animation/LiveLink/Source/LiveLink/Private/Visualizers/LiveLinkDataPreviewComponent.cpp
--- a/Engine/Plugins/Animation/LiveLink/Source/LiveLink/Private/Visualizers/LiveLinkDataPreviewComponent.cpp
+++ b/Engine/Plugins/Animation/LiveLink/Source/LiveLink/Private/Visualizers/LiveLinkDataPreviewComponent.cpp
@@ -6,7 +6,6 @@
 #include "AnimationCoreLibrary.h"
 #include "CanvasTypes.h"
 #include "ILiveLinkClient.h"
-#include "LiveLinkBoneVisualizer.h"
 #include "LiveLinkClient.h"
 #include "Features/IModularFeatures.h"
 #include "Roles/LiveLinkAnimationRole.h"
diff --git a/Engine/Plugins/VirtualProduction/PerformanceCaptureWorkflow/Config/BasePerformanceCaptureWorkflow.ini b/Engine/Plugins/VirtualProduction/PerformanceCaptureWorkflow/Config/BasePerformanceCaptureWorkflow.ini
--- a/Engine/Plugins/VirtualProduction/PerformanceCaptureWorkflow/Config/BasePerformanceCaptureWorkflow.ini
+++ b/Engine/Plugins/VirtualProduction/PerformanceCaptureWorkflow/Config/BasePerformanceCaptureWorkflow.ini
@@ -14,15 +14,10 @@ DatabaseHelperClass=/PerformanceCaptureWorkflow/Database/EU_InternalDatabaseHelp
 +ClassRedirects=(OldName="/Script/PerformanceCaptureWorkflow.PCap_PropDataAsset",NewName="/Script/PerformanceCaptureWorkflow.PCapPropDataAsset")
 +ClassRedirects=(OldName="/Script/PerformanceCaptureWorkflow.PCap_PerformerDataAsset",NewName="/Script/PerformanceCaptureWorkflow.PCapPerformerDataAsset")
 
-+ClassRedirects=(OldName="/Script/PerformanceCaptureWorkflow.PCapBoneVisualiser",NewName="/Script/LiveLink.LiveLinkBoneVisualizer")
 +ClassRedirects=(OldName="/Script/PerformanceCaptureWorkflow.PCapMarkerVisualizer",NewName="/Script/LiveLink.LiveLinkMarkerVisualizer")
 +ClassRedirects=(OldName="/Script/PerformanceCaptureWorkflow.LiveLinkDataPreview",NewName="/Script/LiveLink.LiveLinkDataPreview")
 +ClassRedirects=(OldName="/Script/PerformanceCaptureWorkflow.LiveLinkDataPreviewComponent",NewName="/Script/LiveLink.LiveLinkDataPreviewComponent")
 
-+PackageRedirects=(OldName="/PerformanceCaptureWorkflow/Visualizers/AxisVisualizer", NewName="/LiveLink/Visualizers/AxisVisualizer")
-+PackageRedirects=(OldName="/PerformanceCaptureWorkflow/Visualizers/BoneVisualizer", NewName="/LiveLink/Visualizers/BoneVisualizer")
-+PackageRedirects=(OldName="/PerformanceCaptureWorkflow/Visualizers/JointVisualizer", NewName="/LiveLink/Visualizers/JointVisualizer")
-
 +PackageRedirects=(OldName="/PerformanceCaptureWorkflow/Visualizers/MI_JointAxisVisualizer_Blue", NewName="/LiveLink/Visualizers/MI_JointAxisVisualizer_Blue")
 +PackageRedirects=(OldName="/PerformanceCaptureWorkflow/Visualizers/MI_JointAxisVisualizer_Green", NewName="/LiveLink/Visualizers/MI_JointAxisVisualizer_Green")
 +PackageRedirects=(OldName="/PerformanceCaptureWorkflow/Visualizers/MI_JointAxisVisualizer_Red", NewName="/LiveLink/Visualizers/MI_JointAxisVisualizer_Red")
diff --git a/Engine/Plugins/VirtualProduction/PerformanceCaptureWorkflow/Source/PerformanceCaptureWorkflow/Private/Visualizers/PCapBoneVisualizer.cpp b/Engine/Plugins/VirtualProduction/PerformanceCaptureWorkflow/Source/PerformanceCaptureWorkflow/Private/Visualizers/PCapBoneVisualizer.cpp
new file mode 100644
--- /dev/null
+++ b/Engine/Plugins/VirtualProduction/PerformanceCaptureWorkflow/Source/PerformanceCaptureWorkflow/Private/Visualizers/PCapBoneVisualizer.cpp
@@ -0,0 +1,148 @@
+// Copyright Epic Games, Inc. All Rights Reserved.
+
+
+#include "PCapBoneVisualizer.h"
+#include "AnimationCoreLibrary.h"
+#include "SkeletalMeshAttributes.h"
+#include "Animation/SkeletalMeshActor.h"
+#include "Components/SkeletalMeshComponent.h"
+#include "Materials/MaterialInstanceDynamic.h"
+
+
+// Sets default values for this component's properties
+UPCapBoneVisualiser::UPCapBoneVisualiser(const FObjectInitializer& ObjectInitializer)
+{
+	PrimaryComponentTick.bCanEverTick = true;
+	bTickInEditor = true; //make this component tick in editor
+	VisualizationType = EBoneVisualType::Joint;
+	Color = FLinearColor(0.5,0.5,0.5,1.0);
+	UPrimitiveComponent::SetCollisionEnabled(ECollisionEnabled::NoCollision);
+	bAlwaysCreatePhysicsState = false;
+	SetCastShadow(false);
+}
+
+void UPCapBoneVisualiser::OnRegister()
+{
+	Super::OnRegister();
+
+	ASkeletalMeshActor* OwnerActor = Cast<ASkeletalMeshActor>(GetOwner());
+	
+	if(OwnerActor)
+	{
+		SkelmeshComponent = OwnerActor->GetSkeletalMeshComponent();
+	}
+
+	ClearInstances();
+
+	if(SkelmeshComponent)
+	{
+		TArray<FTransform> BoneTransforms = GetBoneTransforms(SkelmeshComponent);
+		AddInstances(BoneTransforms, false, true, false);
+
+		if(!DynamicMaterial)
+		{
+			UMaterialInterface* BaseMaterial = GetMaterial(0);
+			if(BaseMaterial) //Check if the BaseMaterial is Valid. It will only be valid if a mesh has already been sent. 
+			{
+				DynamicMaterial = CreateDynamicMaterialInstance(0, BaseMaterial, TEXT("MaterialInstance"));
+				SetMaterial(0, DynamicMaterial);
+				DynamicMaterial->ClearParameterValues();
+				DynamicMaterial->SetVectorParameterValue(FName(TEXT("Color")), Color);
+			}
+		}
+	}
+}
+
+void UPCapBoneVisualiser::OnUnregister()
+{
+	Super::OnUnregister();
+
+	ClearInstances();
+}
+
+void UPCapBoneVisualiser::TickComponent(float DeltaTime, ELevelTick TickType, FActorComponentTickFunction* ThisTickFunction)
+{
+	Super::TickComponent(DeltaTime, TickType, ThisTickFunction);
+	
+	if(SkelmeshComponent)
+	{
+		switch (VisualizationType)
+		{
+		case EBoneVisualType::Bone:
+			{
+				BatchUpdateInstancesTransforms(0, GetBoneTransforms(SkelmeshComponent), true, true, false);
+			}
+			break;
+		case EBoneVisualType::Joint:
+			{
+				BatchUpdateInstancesTransforms(0, GetJointTransforms(SkelmeshComponent), true, true, false);
+			}
+
+			break;
+		default:
+			checkNoEntry();
+		}
+	}
+}
+
+TArray<FTransform> UPCapBoneVisualiser::GetJointTransforms(USkeletalMeshComponent* InSkelMeshComponent) const
+{
+	TArray<FTransform> FoundTransforms;
+
+	int32 i = InSkelMeshComponent->GetNumBones();
+	FTransform BoneTransform;
+
+	for (i = 0; i < InSkelMeshComponent->GetNumBones(); i++ )
+	{
+		BoneTransform = InSkelMeshComponent->GetBoneTransform(InSkelMeshComponent->GetBoneName(i));
+		FoundTransforms.Add(BoneTransform);
+	}
+
+	return FoundTransforms;
+}
+
+TArray<FTransform> UPCapBoneVisualiser::GetBoneTransforms(USkeletalMeshComponent* InSkelMeshComponent) const
+{
+	TArray<FTransform> FoundTransforms;
+
+	int32 i = InSkelMeshComponent->GetNumBones();
+	FTransform BoneTransform;
+	FTransform ParentBoneTransform;
+	double BoneLength;
+	FVector Aim = FVector( 0, 0, 1);
+
+	for (i = 0; i < InSkelMeshComponent->GetNumBones(); i++ )
+	{
+		BoneTransform = InSkelMeshComponent->GetBoneTransform(InSkelMeshComponent->GetBoneName(i));
+		ParentBoneTransform = InSkelMeshComponent->GetBoneTransform(InSkelMeshComponent->GetParentBone(InSkelMeshComponent->GetBoneName(i)));
+		BoneLength = FVector::Distance(BoneTransform.GetLocation(), ParentBoneTransform.GetLocation());
+		BoneTransform.SetScale3D(FVector(1, 1, BoneLength));
+		FQuat DiffRotation = AnimationCore::SolveAim(BoneTransform, ParentBoneTransform.GetLocation(), Aim.GetSafeNormal(), false, FVector(1, 1, 1), float(0));
+		BoneTransform.SetRotation(DiffRotation);
+		FoundTransforms.Add(BoneTransform);
+	}
+	return FoundTransforms;
+}
+
+void UPCapBoneVisualiser::UpdateColor(FLinearColor NewColor)
+{
+	Color = NewColor;
+	if(DynamicMaterial)
+	{
+		DynamicMaterial->SetVectorParameterValue(FName(TEXT("Color")), Color);
+	}
+}
+
+void UPCapBoneVisualiser::PostEditChangeProperty(FPropertyChangedEvent& PropertyChangedEvent)
+{
+	Super::PostEditChangeProperty(PropertyChangedEvent);
+
+	if(PropertyChangedEvent.Property != nullptr)
+	{
+		if(PropertyChangedEvent.Property->GetFName() == GET_MEMBER_NAME_CHECKED(UPCapBoneVisualiser, Color))
+		{
+			UpdateColor(Color);
+		}
+	}
+}
+
diff --git a/Engine/Plugins/VirtualProduction/PerformanceCaptureWorkflow/Source/PerformanceCaptureWorkflow/Private/Visualizers/PCapBoneVisualizer.h b/Engine/Plugins/VirtualProduction/PerformanceCaptureWorkflow/Source/PerformanceCaptureWorkflow/Private/Visualizers/PCapBoneVisualizer.h
new file mode 100644
--- /dev/null
+++ b/Engine/Plugins/VirtualProduction/PerformanceCaptureWorkflow/Source/PerformanceCaptureWorkflow/Private/Visualizers/PCapBoneVisualizer.h
@@ -0,0 +1,70 @@
+// Copyright Epic Games, Inc. All Rights Reserved.
+
+#pragma once
+
+#include "CoreMinimal.h"
+#include "Components/InstancedStaticMeshComponent.h"
+#include "PCapBoneVisualizer.generated.h"
+
+class USkeletalMeshComponent;
+class UMaterialInstanceDynamic;
+
+UENUM(BlueprintType)
+enum class EBoneVisualType : uint8
+{
+	Joint	UMETA(DisplayName = "Joint"),
+	Bone	UMETA(DisplayName = "Bone"),
+};
+
+/** Instanced Static Mesh Component for drawing bones and joints on Skeletal Mesh. Only usable with SkeletalMesh Actors*/
+UCLASS(ClassGroup=(PerformanceCapture), meta=(BlueprintSpawnableComponent),Transient, DisplayName = "PerformanceCaptureVisualizer")
+class PERFORMANCECAPTUREWORKFLOW_API UPCapBoneVisualiser : public UInstancedStaticMeshComponent
+{
+	GENERATED_BODY()
+
+public:
+	// Sets default values for this component's properties
+	UPCapBoneVisualiser(const FObjectInitializer& ObjectInitializer);
+	
+	// Called every frame
+	virtual void TickComponent(float DeltaTime, ELevelTick TickType, FActorComponentTickFunction* ThisTickFunction) override;
+
+	/** Get all the transforms from the skelmesh. */
+	TArray<FTransform> GetJointTransforms(USkeletalMeshComponent* InSkelMeshComponent) const;
+
+	/** Get all the bones' transforms for drawing a joint (point to point mesh). */
+	TArray<FTransform> GetBoneTransforms(USkeletalMeshComponent* InSkelMeshComponent) const;
+	
+	/** Whether to draw joints along the length of the bone, or just at the pivot or scaled by the distance to the parent bone. */
+	UPROPERTY(EditAnywhere, BlueprintReadWrite, Category="PerformanceCapture|Visualization")
+	EBoneVisualType VisualizationType;
+
+	/** Color to use on the instanced static meshes. Can only be set during construction. */
+	UPROPERTY(EditAnywhere, BlueprintReadWrite, Setter="UpdateColor", Category="PerformanceCapture|Visualization")
+	FLinearColor Color;
+
+	/** Dynamic Material to use on the instanced static meshes. */
+	UPROPERTY(EditAnywhere, BlueprintReadOnly, Category="PerformanceCapture|Visualization")
+	TObjectPtr<UMaterialInstanceDynamic> DynamicMaterial;
+
+	/**
+	 * Update the color for this bone visaulizer.
+	 * @param NewColor New color to apply in the material.
+	 */
+	UFUNCTION(BlueprintCallable, Category="PerformanceCapture|Visualization")
+	void UpdateColor(FLinearColor NewColor);
+
+#if WITH_EDITOR
+
+	virtual void PostEditChangeProperty(FPropertyChangedEvent& PropertyChangedEvent) override;
+	
+#endif
+
+protected:
+	virtual void OnRegister() override;
+	virtual void OnUnregister() override;
+
+private:
+	UPROPERTY()
+	TObjectPtr<USkeletalMeshComponent> SkelmeshComponent;
+};
