[Performance Capture Workflow]
-Updated plugin description to aid discoverability in the plugin browser by including text referring to Mocap Manager.

diff --git a/Engine/Plugins/VirtualProduction/PerformanceCaptureWorkflow/PerformanceCaptureWorkflow.uplugin b/Engine/Plugins/VirtualProduction/PerformanceCaptureWorkflow/PerformanceCaptureWorkflow.uplugin
--- a/Engine/Plugins/VirtualProduction/PerformanceCaptureWorkflow/PerformanceCaptureWorkflow.uplugin
+++ b/Engine/Plugins/VirtualProduction/PerformanceCaptureWorkflow/PerformanceCaptureWorkflow.uplugin
@@ -3,7 +3,7 @@
 	"Version": 1,
 	"VersionName": "0.1",
 	"FriendlyName": "Performance Capture Workflow",
-	"Description": "Performance Capture In-Editor Workflow tools.",
+	"Description": "Performance Capture In-Editor Workflow tools. Provides access to the Mocap Manager panel.",
 	"Category": "Animation",
 	"CreatedBy": "Epic Games",
 	"CreatedByURL": "https://epicgames.com",
