[PCap Workflow]
- Fix typos on naming tokens tooltip

#rnx

diff --git a/Engine/Plugins/VirtualProduction/PerformanceCaptureWorkflow/Source/PerformanceCaptureWorkflow/Private/PCapSessionTemplate.h b/Engine/Plugins/VirtualProduction/PerformanceCaptureWorkflow/Source/PerformanceCaptureWorkflow/Private/PCapSessionTemplate.h
--- a/Engine/Plugins/VirtualProduction/PerformanceCaptureWorkflow/Source/PerformanceCaptureWorkflow/Private/PCapSessionTemplate.h
+++ b/Engine/Plugins/VirtualProduction/PerformanceCaptureWorkflow/Source/PerformanceCaptureWorkflow/Private/PCapSessionTemplate.h
@@ -243,11 +243,10 @@ public:
 	"{ddd} - 3-character Day (lowercase)\n"
 	"{dd} - Day (2 digit)\n"
 	"{ampm} - am or pm (lowercase)\n"
-	"{AMPM} - AM or PM (UPPERCASE)n"
+	"{AMPM} - AM or PM (UPPERCASE)\n"
 	"{12h} - Hour (12)\n"
 	"{24h} - Hour (24)\n"
-	"{24h} - Hour (24)\n"
-	"{24h} - Hour (24)\n"
+	"{min} - Minute\n"
 	"{sec} - Second\n"
 	"{ms} - Millisecond";
 
