[Peformance Capture Workflow]
- Added runtime module
- Refactor Actor, actor component and animinstanc classes into new runtime module to enable running tools/actors in -game session.
- Added redirectors to new class locations

[RN] Performance Capture actor and components now usable at runtime for -game workflows.

diff --git a/Engine/Plugins/VirtualProduction/PerformanceCaptureWorkflow/Config/BasePerformanceCaptureWorkflow.ini b/Engine/Plugins/VirtualProduction/PerformanceCaptureWorkflow/Config/BasePerformanceCaptureWorkflow.ini
--- a/Engine/Plugins/VirtualProduction/PerformanceCaptureWorkflow/Config/BasePerformanceCaptureWorkflow.ini
+++ b/Engine/Plugins/VirtualProduction/PerformanceCaptureWorkflow/Config/BasePerformanceCaptureWorkflow.ini
@@ -38,4 +38,9 @@ DatabaseHelperClass=/PerformanceCaptureWorkflow/Database/EU_InternalDatabaseHelp
 +PackageRedirects=(OldName="/PerformanceCaptureWorkflow/Visualizers/SM_MarkerAxis", NewName="/LiveLink/Visualizers/SM_MarkerAxis")
 +PackageRedirects=(OldName="/PerformanceCaptureWorkflow/Visualizers/SM_MarkerBox", NewName="/LiveLink/Visualizers/SM_MarkerBox")
 +PackageRedirects=(OldName="/PerformanceCaptureWorkflow/Visualizers/SM_MarkerSphere", NewName="/LiveLink/Visualizers/SM_MarkerSphere")
-+PackageRedirects=(OldName="/PerformanceCaptureWorkflow/Visualizers/SM_TransformAxis", NewName="/LiveLink/Visualizers/SM_TransformAxis")
\ No newline at end of file
++PackageRedirects=(OldName="/PerformanceCaptureWorkflow/Visualizers/SM_TransformAxis", NewName="/LiveLink/Visualizers/SM_TransformAxis")
+
++ClassRedirects=(OldName="/Script/PerformanceCaptureWorkflow.PerformanceCaptureStageRoot",NewName="/Script/PerformanceCaptureWorkflowRuntime.PerformanceCaptureStageRoot")
++ClassRedirects=(OldName="/Script/PerformanceCaptureWorkflow.PCapPropComponent",NewName="/Script/PerformanceCaptureWorkflowRuntime.PCapPropComponent")
++ClassRedirects=(OldName="/Script/PerformanceCaptureWorkflow.PCapPropLiveLinkAnimInstance",NewName="/Script/PerformanceCaptureWorkflowRuntime.PCapPropLiveLinkAnimInstance")
++ClassRedirects=(OldName="/Script/PerformanceCaptureWorkflow.AnimNode_LiveLinkProp",NewName="/Script/PerformanceCaptureWorkflowRuntime.AnimNode_LiveLinkProp")
\ No newline at end of file
diff --git a/Engine/Plugins/VirtualProduction/PerformanceCaptureWorkflow/PerformanceCaptureWorkflow.uplugin b/Engine/Plugins/VirtualProduction/PerformanceCaptureWorkflow/PerformanceCaptureWorkflow.uplugin
--- a/Engine/Plugins/VirtualProduction/PerformanceCaptureWorkflow/PerformanceCaptureWorkflow.uplugin
+++ b/Engine/Plugins/VirtualProduction/PerformanceCaptureWorkflow/PerformanceCaptureWorkflow.uplugin
@@ -19,6 +19,11 @@
 			"Name": "PerformanceCaptureWorkflow",
 			"Type": "Editor",
 			"LoadingPhase": "Default"
+		},
+		{
+			"Name": "PerformanceCaptureWorkflowRuntime",
+			"Type": "Runtime",
+			"LoadingPhase": "Default"
 		}
 	],
 	"Plugins": [
diff --git a/Engine/Plugins/VirtualProduction/PerformanceCaptureWorkflow/Source/PerformanceCaptureWorkflow/PerformanceCaptureWorkflow.Build.cs b/Engine/Plugins/VirtualProduction/PerformanceCaptureWorkflow/Source/PerformanceCaptureWorkflow/PerformanceCaptureWorkflow.Build.cs
--- a/Engine/Plugins/VirtualProduction/PerformanceCaptureWorkflow/Source/PerformanceCaptureWorkflow/PerformanceCaptureWorkflow.Build.cs
+++ b/Engine/Plugins/VirtualProduction/PerformanceCaptureWorkflow/Source/PerformanceCaptureWorkflow/PerformanceCaptureWorkflow.Build.cs
@@ -51,7 +51,8 @@ public class PerformanceCaptureWorkflow : ModuleRules
 				"TakeRecorderSources", 
 				"TakesCore",
 				"SubobjectDataInterface",
-				"NamingTokens"
+				"NamingTokens",
+				"PerformanceCaptureWorkflowRuntime"
 			}
 			);
 	}
diff --git a/Engine/Plugins/VirtualProduction/PerformanceCaptureWorkflow/Source/PerformanceCaptureWorkflow/Private/AnimNode_LiveLinkProp.cpp b/Engine/Plugins/VirtualProduction/PerformanceCaptureWorkflow/Source/PerformanceCaptureWorkflow/Private/AnimNode_LiveLinkProp.cpp
deleted file mode 100644
diff --git a/Engine/Plugins/VirtualProduction/PerformanceCaptureWorkflow/Source/PerformanceCaptureWorkflow/Private/AnimNode_LiveLinkProp.h b/Engine/Plugins/VirtualProduction/PerformanceCaptureWorkflow/Source/PerformanceCaptureWorkflow/Private/AnimNode_LiveLinkProp.h
deleted file mode 100644
diff --git a/Engine/Plugins/VirtualProduction/PerformanceCaptureWorkflow/Source/PerformanceCaptureWorkflow/Private/PCapPropComponent.cpp b/Engine/Plugins/VirtualProduction/PerformanceCaptureWorkflow/Source/PerformanceCaptureWorkflow/Private/PCapPropComponent.cpp
deleted file mode 100644
diff --git a/Engine/Plugins/VirtualProduction/PerformanceCaptureWorkflow/Source/PerformanceCaptureWorkflow/Private/PCapPropComponent.h b/Engine/Plugins/VirtualProduction/PerformanceCaptureWorkflow/Source/PerformanceCaptureWorkflow/Private/PCapPropComponent.h
deleted file mode 100644
diff --git a/Engine/Plugins/VirtualProduction/PerformanceCaptureWorkflow/Source/PerformanceCaptureWorkflow/Private/PCapPropLiveLinkAnimInstance.cpp b/Engine/Plugins/VirtualProduction/PerformanceCaptureWorkflow/Source/PerformanceCaptureWorkflow/Private/PCapPropLiveLinkAnimInstance.cpp
deleted file mode 100644
diff --git a/Engine/Plugins/VirtualProduction/PerformanceCaptureWorkflow/Source/PerformanceCaptureWorkflow/Private/PCapPropLiveLinkAnimInstance.h b/Engine/Plugins/VirtualProduction/PerformanceCaptureWorkflow/Source/PerformanceCaptureWorkflow/Private/PCapPropLiveLinkAnimInstance.h
deleted file mode 100644
diff --git a/Engine/Plugins/VirtualProduction/PerformanceCaptureWorkflow/Source/PerformanceCaptureWorkflow/Private/PCapStageRoot.cpp b/Engine/Plugins/VirtualProduction/PerformanceCaptureWorkflow/Source/PerformanceCaptureWorkflow/Private/PCapStageRoot.cpp
deleted file mode 100644
diff --git a/Engine/Plugins/VirtualProduction/PerformanceCaptureWorkflow/Source/PerformanceCaptureWorkflow/Public/PCapStageRoot.h b/Engine/Plugins/VirtualProduction/PerformanceCaptureWorkflow/Source/PerformanceCaptureWorkflow/Public/PCapStageRoot.h
deleted file mode 100644
diff --git a/Engine/Plugins/VirtualProduction/PerformanceCaptureWorkflow/Source/PerformanceCaptureWorkflowRuntime/PerformanceCaptureWorkflowRuntime.Build.cs b/Engine/Plugins/VirtualProduction/PerformanceCaptureWorkflow/Source/PerformanceCaptureWorkflowRuntime/PerformanceCaptureWorkflowRuntime.Build.cs
new file mode 100644
--- /dev/null
+++ b/Engine/Plugins/VirtualProduction/PerformanceCaptureWorkflow/Source/PerformanceCaptureWorkflowRuntime/PerformanceCaptureWorkflowRuntime.Build.cs
@@ -0,0 +1,31 @@
+﻿// Copyright Epic Games, Inc. All Rights Reserved.
+using UnrealBuildTool;
+
+public class PerformanceCaptureWorkflowRuntime : ModuleRules
+{
+    public PerformanceCaptureWorkflowRuntime(ReadOnlyTargetRules Target) : base(Target)
+    {
+        PCHUsage = ModuleRules.PCHUsageMode.UseExplicitOrSharedPCHs;
+        ShortName = "PCWFR";
+
+        PublicDependencyModuleNames.AddRange(
+            new string[]
+            {
+                "Core",
+            }
+        );
+
+        PrivateDependencyModuleNames.AddRange(
+            new string[]
+            {
+                "CoreUObject",
+                "Engine",
+                "Slate",
+                "SlateCore",
+                "LiveLink",
+                "LiveLinkInterface",
+                "LiveLinkAnimationCore",
+            }
+        );
+    }
+}
\ No newline at end of file
diff --git a/Engine/Plugins/VirtualProduction/PerformanceCaptureWorkflow/Source/PerformanceCaptureWorkflowRuntime/Private/AnimNode_LiveLinkProp.cpp b/Engine/Plugins/VirtualProduction/PerformanceCaptureWorkflow/Source/PerformanceCaptureWorkflowRuntime/Private/AnimNode_LiveLinkProp.cpp
new file mode 100644
--- /dev/null
+++ b/Engine/Plugins/VirtualProduction/PerformanceCaptureWorkflow/Source/PerformanceCaptureWorkflowRuntime/Private/AnimNode_LiveLinkProp.cpp
@@ -0,0 +1,140 @@
+// Copyright Epic Games, Inc. All Rights Reserved.
+
+
+#include "AnimNode_LiveLinkProp.h"
+
+#include "Animation/AnimInstanceProxy.h"
+#include "Animation/AnimTrace.h"
+#include "Features/IModularFeatures.h"
+#include "ILiveLinkClient.h"
+#include "Roles/LiveLinkAnimationRole.h"
+#include "Roles/LiveLinkAnimationTypes.h"
+
+#include UE_INLINE_GENERATED_CPP_BY_NAME(AnimNode_LiveLinkProp)
+
+FAnimNode_LiveLinkProp::FAnimNode_LiveLinkProp()
+	: LiveLinkClient_AnyThread(nullptr)
+	, CachedDeltaTime(0.0f)
+{
+}
+
+void FAnimNode_LiveLinkProp::BuildPoseFromAnimData(const FLiveLinkSubjectFrameData& LiveLinkData, FPoseContext& Output)
+{
+	const FLiveLinkSkeletonStaticData* SkeletonData = LiveLinkData.StaticData.Cast<FLiveLinkSkeletonStaticData>();
+	const FLiveLinkAnimationFrameData* FrameData = LiveLinkData.FrameData.Cast<FLiveLinkAnimationFrameData>();
+	check(SkeletonData);
+	check(FrameData);
+
+	//Get bones names
+	const TArray<FName>& SourceBoneNames = SkeletonData->GetBoneNames();
+
+	//Iterate through bones in LiveLink Data
+	for (int32 i = 0; i < SourceBoneNames.Num(); ++i)
+	{
+		FName BoneName = SourceBoneNames[i];
+
+		FTransform BoneTransform = OffsetTransform*FrameData->Transforms[i];
+
+		int32 MeshIndex = Output.Pose.GetBoneContainer().GetPoseBoneIndexForBoneName(BoneName);
+		if (MeshIndex != INDEX_NONE)
+		{
+			FCompactPoseBoneIndex CompactPoseBoneIndex = Output.Pose.GetBoneContainer().MakeCompactPoseIndex(FMeshPoseBoneIndex(MeshIndex));
+			if(CompactPoseBoneIndex != INDEX_NONE)
+			{
+				Output.Pose[CompactPoseBoneIndex] = BoneTransform;
+			}
+		}
+	}
+	CachedDeltaTime = 0.0f;
+}
+
+void FAnimNode_LiveLinkProp::Initialize_AnyThread(const FAnimationInitializeContext& Context)
+{
+	InputPose.Initialize(Context);
+
+	if (CachedLiveLinkFrameData.IsValid() == false)
+	{
+		CachedLiveLinkFrameData  = MakeShared<FLiveLinkSubjectFrameData>();
+	}
+}
+
+void FAnimNode_LiveLinkProp::PreUpdate(const UAnimInstance* InAnimInstance)
+{
+	ILiveLinkClient* ThisFrameClient = nullptr;
+	IModularFeatures& ModularFeatures  = IModularFeatures::Get();
+	if (ModularFeatures.IsModularFeatureAvailable(ILiveLinkClient::ModularFeatureName))
+	{
+		ThisFrameClient = &IModularFeatures::Get().GetModularFeature<ILiveLinkClient>(ILiveLinkClient::ModularFeatureName);
+	}
+	LiveLinkClient_AnyThread = ThisFrameClient;
+}
+
+void FAnimNode_LiveLinkProp::Update_AnyThread(const FAnimationUpdateContext& Context)
+{
+	InputPose.Update(Context);
+
+	GetEvaluateGraphExposedInputs().Execute(Context);
+
+	CachedDeltaTime += Context.GetDeltaTime();
+
+	TRACE_ANIM_NODE_VALUE(Context, TEXT("SubjectName"), LiveLinkSubjectName.Name);
+}
+
+void FAnimNode_LiveLinkProp::Evaluate_AnyThread(FPoseContext& Output)
+{
+	InputPose.Evaluate(Output);
+
+	if (!LiveLinkClient_AnyThread)
+	{
+		return;
+	}
+
+	FLiveLinkSubjectFrameData SubjectFrameData;
+
+	if(bDoLiveLinkEvaluation)
+	{
+		// Invalidate cached evaluated Role to make sure we have a valid one during the last evaluation when using it
+		CachedEvaluatedRole = nullptr;
+
+		TSubclassOf<ULiveLinkRole> SubjectRole = LiveLinkClient_AnyThread->GetSubjectRole_AnyThread(LiveLinkSubjectName);
+		if (SubjectRole)
+		{
+			if(LiveLinkClient_AnyThread->DoesSubjectSupportsRole_AnyThread(LiveLinkSubjectName, ULiveLinkAnimationRole::StaticClass()))
+			{
+				//Process the animation data
+				if(LiveLinkClient_AnyThread->EvaluateFrame_AnyThread(LiveLinkSubjectName, ULiveLinkAnimationRole::StaticClass(), SubjectFrameData))
+				{
+					BuildPoseFromAnimData(SubjectFrameData, Output);
+					CachedEvaluatedRole = ULiveLinkAnimationRole::StaticClass();
+				}
+			}
+			CachedLiveLinkFrameData->StaticData = MoveTemp(SubjectFrameData.StaticData);
+			CachedLiveLinkFrameData->FrameData = MoveTemp(SubjectFrameData.FrameData);
+		}
+	}
+	else
+	{
+		if(CachedLiveLinkFrameData && CachedEvaluatedRole)
+		{
+			if (CachedEvaluatedRole == ULiveLinkAnimationRole::StaticClass())
+			{
+				BuildPoseFromAnimData(*CachedLiveLinkFrameData, Output);
+			}
+		}
+	}
+}
+
+void FAnimNode_LiveLinkProp::CacheBones_AnyThread(const FAnimationCacheBonesContext& Context)
+{
+	Super::CacheBones_AnyThread(Context);
+	InputPose.CacheBones(Context);
+}
+
+void FAnimNode_LiveLinkProp::GatherDebugData(FNodeDebugData& DebugData)
+{
+	FString DebugLine = FString::Printf(TEXT("LiveLink Prop - SubjectName: %s"), *LiveLinkSubjectName.ToString());
+
+	DebugData.AddDebugItem(DebugLine);
+	InputPose.GatherDebugData(DebugData);
+}
+
diff --git a/Engine/Plugins/VirtualProduction/PerformanceCaptureWorkflow/Source/PerformanceCaptureWorkflowRuntime/Private/AnimNode_LiveLinkProp.h b/Engine/Plugins/VirtualProduction/PerformanceCaptureWorkflow/Source/PerformanceCaptureWorkflowRuntime/Private/AnimNode_LiveLinkProp.h
new file mode 100644
--- /dev/null
+++ b/Engine/Plugins/VirtualProduction/PerformanceCaptureWorkflow/Source/PerformanceCaptureWorkflowRuntime/Private/AnimNode_LiveLinkProp.h
@@ -0,0 +1,70 @@
+// Copyright Epic Games, Inc. All Rights Reserved.
+
+#pragma once
+
+#include "Animation/AnimNodeBase.h"
+
+#include "CoreMinimal.h"
+#include "LiveLinkTypes.h"
+#include "Templates/SubclassOf.h"
+
+#include "AnimNode_LiveLinkProp.generated.h"
+
+class ILiveLinkClient;
+class ULiveLinkRole;
+struct FLiveLinkSubjectFrameData;
+
+/**
+ * This animnode is exclusively for Mocap props - single bone skeleton meshes. Not exposed to the animation graph.
+ */
+
+USTRUCT(BlueprintInternalUseOnly)
+struct FAnimNode_LiveLinkProp : public FAnimNode_Base
+{
+	GENERATED_BODY()
+
+public:
+	/** Input pose */
+	UPROPERTY(EditAnywhere, Category = "Performance Capture|Prop")
+	FPoseLink InputPose;
+
+	/** The Live Link subject to use. */
+	UPROPERTY(EditAnywhere, Category = "Performance Capture|Prop")
+	FLiveLinkSubjectName LiveLinkSubjectName;
+
+	/** Bool to control evaluation of animation. */
+	UPROPERTY(EditAnywhere, Transient, Category = "Performance Capture|Prop")
+	bool bDoLiveLinkEvaluation = true;
+
+	/* Transform to apply local space to offset to the incoming data. */
+	UPROPERTY(EditAnywhere, Category = "Performance Capture|Prop")
+	FTransform OffsetTransform = FTransform::Identity;
+
+	
+
+	FAnimNode_LiveLinkProp();
+
+	//~ FAnimNode_Base interface
+	virtual void Initialize_AnyThread(const FAnimationInitializeContext& Context) override;
+	virtual void PreUpdate(const UAnimInstance* InAnimInstance) override;
+	virtual void Update_AnyThread(const FAnimationUpdateContext & Context) override;
+	virtual void Evaluate_AnyThread(FPoseContext& Output) override;
+	virtual void CacheBones_AnyThread(const FAnimationCacheBonesContext & Context) override;
+	virtual void GatherDebugData(FNodeDebugData& DebugData) override;
+	virtual bool HasPreUpdate() const override { return true; };
+	//~ End of FAnimNode_Base interface
+
+protected:
+	
+	void BuildPoseFromAnimData(const FLiveLinkSubjectFrameData& LiveLinkData, FPoseContext& Output);
+
+private:
+
+	ILiveLinkClient* LiveLinkClient_AnyThread;
+
+	float CachedDeltaTime;
+
+	TSharedPtr<FLiveLinkSubjectFrameData> CachedLiveLinkFrameData;
+
+	TSubclassOf<class ULiveLinkRole> CachedEvaluatedRole;
+};
diff --git a/Engine/Plugins/VirtualProduction/PerformanceCaptureWorkflow/Source/PerformanceCaptureWorkflowRuntime/Private/PCapPropComponent.cpp b/Engine/Plugins/VirtualProduction/PerformanceCaptureWorkflow/Source/PerformanceCaptureWorkflowRuntime/Private/PCapPropComponent.cpp
new file mode 100644
--- /dev/null
+++ b/Engine/Plugins/VirtualProduction/PerformanceCaptureWorkflow/Source/PerformanceCaptureWorkflowRuntime/Private/PCapPropComponent.cpp
@@ -0,0 +1,269 @@
+// Copyright Epic Games, Inc. All Rights Reserved.
+
+
+#include "PCapPropComponent.h"
+#include "Components/SkeletalMeshComponent.h"
+#include "ILiveLinkClient.h"
+#include "PCapPropLiveLinkAnimInstance.h"
+#include "GameFramework/Actor.h"
+#include "Components/SceneComponent.h"
+#include "Features/IModularFeatures.h"
+#include "Roles/LiveLinkAnimationRole.h"
+#include "Roles/LiveLinkTransformRole.h"
+#include "Roles/LiveLinkTransformTypes.h"
+
+
+// Sets default values for this component's properties
+UPCapPropComponent::UPCapPropComponent()
+	: bIsDirty(true)
+{
+	PrimaryComponentTick.bCanEverTick = true;
+	PrimaryComponentTick.bStartWithTickEnabled = true;
+	PrimaryComponentTick.TickGroup = ETickingGroup::TG_PrePhysics;
+	bTickInEditor = true;
+}
+
+UPCapPropComponent::~UPCapPropComponent()
+{
+	//TODO: Might have to do something when exiting PIE here
+}
+
+USceneComponent* UPCapPropComponent::GetControlledComponent() const
+{
+	if(ControlledComponent.GetComponent(GetOwner()))
+	{
+		return Cast<USceneComponent>(ControlledComponent.GetComponent(GetOwner()));
+	}
+
+	return nullptr;
+}
+
+void UPCapPropComponent::SetControlledComponent(USceneComponent* InComponent)
+{
+	//Check the offered new component is inside the same actor as this prop component
+	if(InComponent->GetOwner() == this->GetOwner())
+	{
+		ControlledComponent.OverrideComponent = InComponent;
+	}
+}
+
+void UPCapPropComponent::SetLiveLinkSubject(FLiveLinkSubjectName Subject)
+{
+	SubjectName = Subject;
+
+	//If there is a Skeletal Mesh controlled by an Anim Instance, reinit.
+	InitiateAnimation();
+}
+
+FLiveLinkSubjectName UPCapPropComponent::GetLiveLinkSubject() const
+{
+	return SubjectName;
+}
+
+void UPCapPropComponent::SetEvaluateLiveLinkData(bool bEvaluate)
+{
+	bEvaluateLiveLink = bEvaluate;
+
+	//If there is a Skeletal Mesh controlled by an Anim Instance, reinit.
+	USkeletalMeshComponent* ControlledMesh = Cast<USkeletalMeshComponent>(ControlledComponent.GetComponent(GetOwner()));
+	if(IsValid(ControlledMesh))
+	{
+		UPCapPropLiveLinkAnimInstance* LiveLinkAnimInstance;
+		LiveLinkAnimInstance = Cast<UPCapPropLiveLinkAnimInstance>(ControlledMesh->GetAnimInstance());
+		if(LiveLinkAnimInstance)
+		{
+			LiveLinkAnimInstance->EnableLiveLinkEvaluation(bEvaluateLiveLink);
+		}
+	}
+}
+
+bool UPCapPropComponent::GetEvaluateLiveLinkData()
+{
+	return bEvaluateLiveLink;
+}
+
+void UPCapPropComponent::SetOffsetTransform(FTransform NewOffset)
+{
+	OffsetTransform = NewOffset;
+	
+	//If there is a Skeletal Mesh controlled by an Anim Instance, reinit.
+	TObjectPtr<USkeletalMeshComponent> ControlledMesh = Cast<USkeletalMeshComponent>(ControlledComponent.GetComponent(GetOwner()));
+	if(IsValid(ControlledMesh))
+	{
+		UPCapPropLiveLinkAnimInstance* LiveLinkAnimInstance;
+		LiveLinkAnimInstance = Cast<UPCapPropLiveLinkAnimInstance>(ControlledMesh->GetAnimInstance());
+		if(LiveLinkAnimInstance)
+		{
+			LiveLinkAnimInstance->SetOffsetTransform(OffsetTransform);
+		}
+	}
+}
+
+void UPCapPropComponent::OnRegister()
+{
+	Super::OnRegister();
+	bIsDirty = true;
+
+	GetOwner()->Tags.AddUnique(TEXT("PCapProp"));
+
+	//Enforce mobility on the controlled component
+	if(ControlledComponent.GetComponent(GetOwner()))
+	{
+		Cast<USceneComponent>(ControlledComponent.GetComponent(GetOwner()))->SetMobility(EComponentMobility::Movable);
+	}
+
+	TObjectPtr<USkeletalMeshComponent> ControlledMesh = Cast<USkeletalMeshComponent>(ControlledComponent.GetComponent(GetOwner()));
+	if(ControlledMesh && bIsDirty)
+	{
+		InitiateAnimation();
+	}
+}
+
+void UPCapPropComponent::DestroyComponent(bool bPromoteChildren)
+{
+	ResetAnimInstance();
+	
+	Super::DestroyComponent(bPromoteChildren);
+}
+
+#if WITH_EDITOR
+void UPCapPropComponent::PostEditChangeProperty(FPropertyChangedEvent& PropertyChangedEvent)
+{
+	if(ControlledComponent.GetComponent(GetOwner()))
+	{
+		Cast<USceneComponent>(ControlledComponent.GetComponent(GetOwner()))->SetMobility(EComponentMobility::Movable);
+	}
+	if(PropertyChangedEvent.Property->GetFName() == GET_MEMBER_NAME_CHECKED(UPCapPropComponent, OffsetTransform))
+	{
+		SetOffsetTransform(OffsetTransform);
+	}
+	if(PropertyChangedEvent.Property->GetFName() == GET_MEMBER_NAME_CHECKED(UPCapPropComponent, ControlledComponent))
+	{
+		ResetAnimInstance();
+		bIsDirty = true;
+	}
+	if(PropertyChangedEvent.Property->GetFName() == GET_MEMBER_NAME_CHECKED(UPCapPropComponent, SubjectName))
+	{
+		ResetAnimInstance();
+		bIsDirty = true;
+	}
+
+	Super::PostEditChangeProperty(PropertyChangedEvent);
+}
+#endif
+
+void UPCapPropComponent::TickComponent(float DeltaTime, ELevelTick TickType, FActorComponentTickFunction* ThisTickFunction)
+{
+
+	//Check if spawnable
+	if(!bIsSpawnableCache.IsSet() || bIsDirty)
+	{
+		static const FName SequencerActorTag(TEXT("SequencerActor"));
+		const AActor* OwningActor = GetOwner();
+
+		bIsSpawnableCache = OwningActor && OwningActor->ActorHasTag(SequencerActorTag);
+
+		if(*bIsSpawnableCache)
+		{
+			bEvaluateLiveLink = false;
+			bIsControlledBySequencer = true;
+		}
+	}
+
+	if(!bIsControlledBySequencer)
+	{
+		// If we are controlling a static mesh, follow this path of applying a relative transform to the controlled mesh from the LiveLink data
+		const TObjectPtr<UStaticMeshComponent> ControlledMesh = Cast<UStaticMeshComponent>(ControlledComponent.GetComponent(GetOwner()));
+		if(ControlledMesh)
+		{
+			FLiveLinkSubjectFrameData SubjectFrameData;
+
+			ILiveLinkClient& LiveLinkClient = IModularFeatures::Get().GetModularFeature<ILiveLinkClient>(ILiveLinkClient::ModularFeatureName);
+	
+			LiveLinkClient.EvaluateFrame_AnyThread(SubjectName, LiveLinkClient.GetSubjectRole_AnyThread(SubjectName), SubjectFrameData);
+	
+			if(LiveLinkClient.GetSubjectRole_AnyThread(SubjectName) == ULiveLinkAnimationRole::StaticClass())
+			{
+				const FLiveLinkAnimationFrameData* FrameData = SubjectFrameData.FrameData.Cast<FLiveLinkAnimationFrameData>();
+				if(FrameData && bEvaluateLiveLink)
+				{
+					CachedLiveLinkTransform = FrameData->Transforms[0];
+				}
+				GetControlledComponent()->SetRelativeTransform(OffsetTransform * CachedLiveLinkTransform, false, nullptr, ETeleportType::None);
+
+			}
+			else if(LiveLinkClient.GetSubjectRole_AnyThread(SubjectName) == ULiveLinkTransformRole::StaticClass())
+			{
+				const FLiveLinkTransformFrameData* FrameData = SubjectFrameData.FrameData.Cast<FLiveLinkTransformFrameData>();
+				if(FrameData && bEvaluateLiveLink)
+				{
+					CachedLiveLinkTransform = FrameData->Transform;
+				}
+				GetControlledComponent()->SetRelativeTransform(OffsetTransform * CachedLiveLinkTransform, false, nullptr, ETeleportType::None);
+			}
+		}
+		//If we are controlling a Skelmesh and receiving a transform role, apply this route. 
+		const TObjectPtr<USkeletalMeshComponent> ControlledSkelmesh = Cast<USkeletalMeshComponent>(ControlledComponent.GetComponent(GetOwner()));
+		if(ControlledSkelmesh)
+		{
+			FLiveLinkSubjectFrameData SkelSubjectFrameData;
+
+			ILiveLinkClient& SkelLiveLinkClient = IModularFeatures::Get().GetModularFeature<ILiveLinkClient>(ILiveLinkClient::ModularFeatureName);
+
+			SkelLiveLinkClient.EvaluateFrame_AnyThread(SubjectName, SkelLiveLinkClient.GetSubjectRole_AnyThread(SubjectName), SkelSubjectFrameData);
+			if(SkelLiveLinkClient.GetSubjectRole_AnyThread(SubjectName) == ULiveLinkTransformRole::StaticClass())
+			{
+				const FLiveLinkTransformFrameData* FrameData = SkelSubjectFrameData.FrameData.Cast<FLiveLinkTransformFrameData>();
+				if(FrameData && bEvaluateLiveLink)
+				{
+					CachedLiveLinkTransform = FrameData->Transform;
+				}
+				GetControlledComponent()->SetRelativeTransform(OffsetTransform * CachedLiveLinkTransform, false, nullptr, ETeleportType::None);
+			}
+			if((SkelLiveLinkClient.GetSubjectRole_AnyThread(SubjectName) == ULiveLinkAnimationRole::StaticClass()) && bIsDirty)
+			{
+				InitiateAnimation();
+			}
+		}
+	}
+	Super::TickComponent(DeltaTime, TickType, ThisTickFunction);
+}
+
+void UPCapPropComponent::InitiateAnimation()
+{
+	TObjectPtr<USkeletalMeshComponent> ControlledMesh = Cast<USkeletalMeshComponent>(ControlledComponent.GetComponent(GetOwner()));
+	ILiveLinkClient& LiveLinkClient = IModularFeatures::Get().GetModularFeature<ILiveLinkClient>(ILiveLinkClient::ModularFeatureName);
+	
+	//If we are controlling a Skeletal Mesh and the Live Link data is of the animation role, initialize animation instance
+	if(IsValid(ControlledMesh) && LiveLinkClient.GetSubjectRole_AnyThread(SubjectName) == ULiveLinkAnimationRole::StaticClass())
+	{
+		ControlledMesh->SetAnimInstanceClass(UPCapPropLiveLinkAnimInstance::StaticClass());
+		ControlledMesh->InitAnim(true /*bForceReinit*/);
+		ControlledMesh->SetUpdateAnimationInEditor(true);
+		ControlledMesh->bPropagateCurvesToFollowers = true;
+		ControlledMesh->VisibilityBasedAnimTickOption = EVisibilityBasedAnimTickOption::AlwaysTickPoseAndRefreshBones;
+		if(ControlledMesh)
+		{
+			UPCapPropLiveLinkAnimInstance* AnimInstance;
+			AnimInstance = Cast<UPCapPropLiveLinkAnimInstance>(ControlledMesh->GetAnimInstance());
+			if(AnimInstance)
+			{
+				AnimInstance->SetSubject(SubjectName);
+				AnimInstance->EnableLiveLinkEvaluation(bEvaluateLiveLink);
+				AnimInstance->SetOffsetTransform(OffsetTransform);
+				bIsDirty=false;
+			}
+		}
+	}
+}
+
+void UPCapPropComponent::ResetAnimInstance() const
+{
+	TObjectPtr<USkeletalMeshComponent> ControlledMesh = Cast<USkeletalMeshComponent>(ControlledComponent.GetComponent(GetOwner()));
+	if(ControlledMesh)
+	{
+		ControlledMesh->SetAnimInstanceClass(nullptr);
+		ControlledMesh->InitAnim(true /*bForceReinit*/);
+		ControlledMesh->SetUpdateAnimationInEditor(true);
+	}
+}
diff --git a/Engine/Plugins/VirtualProduction/PerformanceCaptureWorkflow/Source/PerformanceCaptureWorkflowRuntime/Private/PCapPropLiveLinkAnimInstance.cpp b/Engine/Plugins/VirtualProduction/PerformanceCaptureWorkflow/Source/PerformanceCaptureWorkflowRuntime/Private/PCapPropLiveLinkAnimInstance.cpp
new file mode 100644
--- /dev/null
+++ b/Engine/Plugins/VirtualProduction/PerformanceCaptureWorkflow/Source/PerformanceCaptureWorkflowRuntime/Private/PCapPropLiveLinkAnimInstance.cpp
@@ -0,0 +1,55 @@
+// Copyright Epic Games, Inc. All Rights Reserved.
+
+
+#include "PCapPropLiveLinkAnimInstance.h"
+
+#include UE_INLINE_GENERATED_CPP_BY_NAME(PCapPropLiveLinkAnimInstance)
+
+/*------------------------------------------------------------------------------
+	Anim Instance Proxy.
+------------------------------------------------------------------------------*/
+
+
+void FPCapPropLiveLinkAnimInstanceProxy::Initialize(UAnimInstance* InAnimInstance)
+{
+	FAnimInstanceProxy::Initialize(InAnimInstance);
+	FAnimationInitializeContext InitContext(this);
+	PoseNode.Initialize_AnyThread(InitContext);
+}
+
+void FPCapPropLiveLinkAnimInstanceProxy::PreUpdate(UAnimInstance* InAnimInstance, float DeltaSeconds)
+{
+	Super::PreUpdate(InAnimInstance, DeltaSeconds);
+
+	if(PoseNode.HasPreUpdate())
+	{
+		PoseNode.PreUpdate(InAnimInstance);
+	}
+}
+
+bool FPCapPropLiveLinkAnimInstanceProxy::Evaluate(FPoseContext& Output)
+{
+	PoseNode.Evaluate_AnyThread(Output);
+	
+	return true;
+}
+
+void FPCapPropLiveLinkAnimInstanceProxy::UpdateAnimationNode(const FAnimationUpdateContext& InContext)
+{
+	UpdateCounter.Increment();
+
+	PoseNode.Update_AnyThread(InContext);
+}
+
+/*------------------------------------------------------------------------------
+	Anim Instance
+------------------------------------------------------------------------------*/
+FAnimInstanceProxy* UPCapPropLiveLinkAnimInstance::CreateAnimInstanceProxy()
+{
+	return new FPCapPropLiveLinkAnimInstanceProxy(this);
+}
+
+void UPCapPropLiveLinkAnimInstance::DestroyAnimInstanceProxy(FAnimInstanceProxy* InProxy)
+{
+	Super::DestroyAnimInstanceProxy(InProxy);
+}
diff --git a/Engine/Plugins/VirtualProduction/PerformanceCaptureWorkflow/Source/PerformanceCaptureWorkflowRuntime/Private/PCapPropLiveLinkAnimInstance.h b/Engine/Plugins/VirtualProduction/PerformanceCaptureWorkflow/Source/PerformanceCaptureWorkflowRuntime/Private/PCapPropLiveLinkAnimInstance.h
new file mode 100644
--- /dev/null
+++ b/Engine/Plugins/VirtualProduction/PerformanceCaptureWorkflow/Source/PerformanceCaptureWorkflowRuntime/Private/PCapPropLiveLinkAnimInstance.h
@@ -0,0 +1,76 @@
+// Copyright Epic Games, Inc. All Rights Reserved.
+
+#pragma once
+
+#include "AnimNode_LiveLinkProp.h"
+#include "Animation/AnimInstanceProxy.h"
+#include "Animation/AnimInstance.h"
+#include "PCapPropLiveLinkAnimInstance.generated.h"
+
+
+USTRUCT()
+struct FPCapPropLiveLinkAnimInstanceProxy : public FAnimInstanceProxy
+{
+public:
+	friend struct FAnimNode_LiveLinkProp;
+
+	GENERATED_BODY()
+
+	FPCapPropLiveLinkAnimInstanceProxy()
+	{
+	}
+
+	FPCapPropLiveLinkAnimInstanceProxy(UAnimInstance* InAnimInstance)
+		: FAnimInstanceProxy(InAnimInstance)
+	{
+	}
+
+	virtual void Initialize(UAnimInstance* InAnimInstance) override;
+	virtual void PreUpdate(UAnimInstance* InAnimInstance, float DeltaSeconds) override;
+	virtual bool Evaluate(FPoseContext& Output) override;
+	virtual void UpdateAnimationNode(const FAnimationUpdateContext& InContext) override;
+
+	UPROPERTY(EditAnywhere, Category = Settings)
+	FAnimNode_LiveLinkProp PoseNode;
+};
+/**
+ * 
+ */
+UCLASS(transient, NotBlueprintable)
+class PERFORMANCECAPTUREWORKFLOWRUNTIME_API UPCapPropLiveLinkAnimInstance : public UAnimInstance
+{
+	GENERATED_BODY()
+
+public:
+
+	UFUNCTION(BlueprintCallable, Category = "Performance Capture|Prop")
+	void SetSubject(FLiveLinkSubjectName SubjectName)
+	{
+		GetProxyOnGameThread<FPCapPropLiveLinkAnimInstanceProxy>().PoseNode.LiveLinkSubjectName = SubjectName;
+	}
+
+	UFUNCTION(BlueprintCallable, Category = "Performance Capture|Prop")
+	void EnableLiveLinkEvaluation(bool bDoEnable)
+	{
+		GetProxyOnGameThread<FPCapPropLiveLinkAnimInstanceProxy>().PoseNode.bDoLiveLinkEvaluation = bDoEnable;
+	}
+	
+	UFUNCTION(BlueprintCallable, Category = "Performance Capture|Prop")
+	void SetOffsetTransform(FTransform Offset)
+	{
+		GetProxyOnGameThread<FPCapPropLiveLinkAnimInstanceProxy>().PoseNode.OffsetTransform = Offset;
+	}
+	
+	//Not exposed to Blueprint
+	bool GetEnableLiveLinkEvaluation()
+	{
+		return GetProxyOnGameThread<FPCapPropLiveLinkAnimInstanceProxy>().PoseNode.bDoLiveLinkEvaluation;
+	}
+
+	protected:
+	virtual FAnimInstanceProxy* CreateAnimInstanceProxy() override;
+	virtual void DestroyAnimInstanceProxy(FAnimInstanceProxy* InProxy) override;
+
+	friend FPCapPropLiveLinkAnimInstanceProxy;
+	
+};
diff --git a/Engine/Plugins/VirtualProduction/PerformanceCaptureWorkflow/Source/PerformanceCaptureWorkflowRuntime/Private/PCapStageRoot.cpp b/Engine/Plugins/VirtualProduction/PerformanceCaptureWorkflow/Source/PerformanceCaptureWorkflowRuntime/Private/PCapStageRoot.cpp
new file mode 100644
--- /dev/null
+++ b/Engine/Plugins/VirtualProduction/PerformanceCaptureWorkflow/Source/PerformanceCaptureWorkflowRuntime/Private/PCapStageRoot.cpp
@@ -0,0 +1,56 @@
+// Copyright Epic Games, Inc. All Rights Reserved.
+
+
+#include "PCapStageRoot.h"
+#include "UObject/ConstructorHelpers.h"
+#include "Engine/TextureRenderTarget2D.h"
+#include "Materials/MaterialInterface.h"
+
+
+
+// Sets default values
+APerformanceCaptureStageRoot::APerformanceCaptureStageRoot(const FObjectInitializer& ObjectInitializer)
+	: Super(ObjectInitializer)
+{
+	// Set this actor to call Tick() every frame.  You can turn this off to improve performance if you don't need it.
+	PrimaryActorTick.bCanEverTick = true;
+
+	USceneComponent* SceneRoot = CreateDefaultSubobject<USceneComponent>("Root");
+	SetRootComponent(SceneRoot);
+	MapCaptureComponent = CreateDefaultSubobject<USceneCaptureComponent2D>("MapCaptureComponent");
+	MapCaptureComponent->SetupAttachment(GetRootComponent());
+	MapCaptureComponent->SetRelativeLocationAndRotation(FVector(0.0f, 0.0f, 500.0f), FRotator(-90.0f, 0.0f, -90.0f));
+	
+	// Set the default properties we want on the map scene capture. By default, do not run on every frame
+	MapCaptureComponent->bCaptureEveryFrame = false;
+	MapCaptureComponent->bCaptureOnMovement = true;
+	MapCaptureComponent->CaptureSource = SCS_BaseColor;
+	MapCaptureComponent->ProjectionType = ECameraProjectionMode::Orthographic;
+	MapCaptureComponent->OrthoWidth = 1024.0f;
+	MapCaptureComponent->bUpdateOrthoPlanes = true;
+	
+	// Load and reference a render target in the plugin content
+	ConstructorHelpers::FObjectFinder<UTextureRenderTarget2D> MapTexture(TEXT("/PerformanceCaptureWorkflow/Stage/RT_OrthoView.RT_OrthoView"));
+	if (MapTexture.Succeeded())
+	{
+		MapCaptureComponent->TextureTarget = MapTexture.Object;
+	}
+
+	DecalComponent = CreateDefaultSubobject<UDecalComponent>("Decal");
+	DecalComponent->SetupAttachment(GetRootComponent());
+	DecalComponent->SetRelativeLocationAndRotation(FVector(0.0f, 0.0f, -100.0f), FRotator(0.0f, 0.0f, -90.0f));
+	
+	// Load and reference the 1m grid decal material
+	ConstructorHelpers::FObjectFinder<UMaterialInterface> DecalMaterial(TEXT("/PerformanceCaptureWorkflow/Stage/MI_ProcGrid.MI_ProcGrid"));
+	if (DecalMaterial.Succeeded())
+	{
+		DecalComponent->SetDecalMaterial(DecalMaterial.Object);
+	}
+	
+	DecalComponent->DecalSize = FVector(256.0f, 256.0f, 256.0f);
+	DecalComponent->SetVisibility(false);
+
+	// Create a scene component to which all the stage meshes can be parented
+	StageMeshParent = CreateDefaultSubobject<USceneComponent>("StageMeshParent");
+	StageMeshParent->SetupAttachment(GetRootComponent());
+}
\ No newline at end of file
diff --git a/Engine/Plugins/VirtualProduction/PerformanceCaptureWorkflow/Source/PerformanceCaptureWorkflowRuntime/Private/PerformanceCaptureWorkflowRuntime.cpp b/Engine/Plugins/VirtualProduction/PerformanceCaptureWorkflow/Source/PerformanceCaptureWorkflowRuntime/Private/PerformanceCaptureWorkflowRuntime.cpp
new file mode 100644
--- /dev/null
+++ b/Engine/Plugins/VirtualProduction/PerformanceCaptureWorkflow/Source/PerformanceCaptureWorkflowRuntime/Private/PerformanceCaptureWorkflowRuntime.cpp
@@ -0,0 +1,18 @@
+﻿// Copyright Epic Games, Inc. All Rights Reserved.
+#include "PerformanceCaptureWorkflowRuntime.h"
+
+#define LOCTEXT_NAMESPACE "FPerformanceCaptureWorkflowRuntimeModule"
+
+void FPerformanceCaptureWorkflowRuntimeModule::StartupModule()
+{
+    
+}
+
+void FPerformanceCaptureWorkflowRuntimeModule::ShutdownModule()
+{
+    
+}
+
+#undef LOCTEXT_NAMESPACE
+    
+IMPLEMENT_MODULE(FPerformanceCaptureWorkflowRuntimeModule, PerformanceCaptureWorkflowRuntime)
\ No newline at end of file
diff --git a/Engine/Plugins/VirtualProduction/PerformanceCaptureWorkflow/Source/PerformanceCaptureWorkflowRuntime/Public/PCapPropComponent.h b/Engine/Plugins/VirtualProduction/PerformanceCaptureWorkflow/Source/PerformanceCaptureWorkflowRuntime/Public/PCapPropComponent.h
new file mode 100644
--- /dev/null
+++ b/Engine/Plugins/VirtualProduction/PerformanceCaptureWorkflow/Source/PerformanceCaptureWorkflowRuntime/Public/PCapPropComponent.h
@@ -0,0 +1,136 @@
+// Copyright Epic Games, Inc. All Rights Reserved.
+
+#pragma once
+
+#include "CoreMinimal.h"
+#include "ILiveLinkClient.h"
+#include "Components/SceneComponent.h"
+#include "Components/ActorComponent.h"
+#include "Components/StaticMeshComponent.h"
+#include "PCapPropComponent.generated.h"
+
+/**
+* A component for use in Motion Capture props. Accepts data as either a LiveLink Transform Role or an Animation Role. If the Live Link data is in the Animation role, this component will take the root bone transform and apply that transform to the owning actor's root component. Can be applied to Static Meshes, Skeletal Meshes and Blueprint constructions. Transform data can be offset in the component's local space. If this component is driving a Skeletal Mesh component and is receiving Live Link Animation data the full bone hierarchy will be applied. 
+*/
+UCLASS(Blueprintable, ClassGroup=("Performance Capture"), meta=(BlueprintSpawnableComponent, DisplayName="Prop Component"))
+class PERFORMANCECAPTUREWORKFLOWRUNTIME_API UPCapPropComponent : public UActorComponent
+{
+	GENERATED_BODY()
+
+public:
+	// Constructor
+	UPCapPropComponent();
+	// Destructor
+	~UPCapPropComponent();
+
+	/**
+	 * LiveLink Subject Name. Must have either Animation Role Type.
+	 */
+	UPROPERTY(EditAnywhere, Category = "Performance Capture")
+	FLiveLinkSubjectName SubjectName;
+
+	/**
+	 * Should LiveLink Subject data be evaluated.
+	 */
+	UPROPERTY(EditAnywhere, Category = "Performance Capture")
+	bool bEvaluateLiveLink = true;
+
+	/**
+	 * Overrides all LiveLink data and the offset transform. Must be manually set on possessable bindings.
+	 */
+	UPROPERTY(BlueprintReadOnly, EditInstanceOnly, Category = "Performance Capture")
+	bool bIsControlledBySequencer = false;
+
+	/**
+	 *The package name of the spawning data asset, if there is one. 
+	 */
+	UPROPERTY(BlueprintReadOnly, Category = "Performance Capture")
+	FName SpawningDataAsset;
+	
+	/**
+	 * Offset the incoming LiveLink Transform data in the local space of the controlled component.
+	 */
+	UPROPERTY(BlueprintReadOnly, EditAnywhere, Category = "Performance Capture")
+	FTransform OffsetTransform = FTransform::Identity;
+
+	/**
+	 * Scene Component driven by the LiveLink data. Defaults to Root component of the owning actor.
+	 */
+	UPROPERTY(EditInstanceOnly, Category = "Performance Capture", meta=(UseComponentPicker, AllowedClasses = "/Script/Engine.SceneComponent"))
+	FComponentReference ControlledComponent;
+
+	/**
+	 * Returns the component controlled by this component. Can return null.
+	 */
+	UFUNCTION(BlueprintPure, Category = "Performance Capture")
+	USceneComponent* GetControlledComponent() const;
+
+	/**
+	 * Sets the component controlled by this component. Component must be within the same actor as this Prop component.
+	 */
+	UFUNCTION(BlueprintCallable, Category = "Performance Capture")
+	void SetControlledComponent(USceneComponent* InComponent);
+
+	/**
+	 * Set the Live Link subject used by this prop component. 
+	 * @param Subject Live Link Subject Name
+	 */
+	UFUNCTION(BlueprintCallable, Category = "Performance Capture|Prop")
+	void SetLiveLinkSubject(FLiveLinkSubjectName Subject);
+	
+	/**
+	* Get the LiveLink Subject Name.
+	* @return FLiveLinkSubjectName Current LiveLink Subject.
+	*/
+	UFUNCTION(BlueprintCallable, Category ="Performance Capture|Prop")
+	FLiveLinkSubjectName GetLiveLinkSubject() const;
+
+	/**
+	* Set the LiveLink data to update the Skeletal Mesh pose.
+	* @param bEvaluateLinkLink Drive or pause the Controlled Mesh from LiveLink Subject data.
+	*/
+	UFUNCTION(BlueprintCallable, Category ="Performance Capture|Prop")
+	void SetEvaluateLiveLinkData(bool bEvaluateLinkLink);
+	
+	/**
+	* Get the LiveLink Evaluation State. Subject must have the Animation Role Type.
+	* @return bool Is LiveLink data being evaluated.
+	*/
+	UFUNCTION(BlueprintPure, Category ="Performance Capture|Prop")
+	bool GetEvaluateLiveLinkData();
+
+	/**
+	* Update the local space offset to this prop's Live Link pose. 
+	* @param NewOffset Offset Transform.
+	*/
+	UFUNCTION(BlueprintCallable, Category ="Performance Capture|Prop")
+	void SetOffsetTransform(FTransform NewOffset);
+
+
+protected:
+
+	//~ Begin UActorComponent interface
+	virtual void OnRegister() override;
+	virtual void DestroyComponent(bool bPromoteChildren) override;
+	//~ End UActorComponent interface
+	
+	//~ UObject interface
+#if WITH_EDITOR
+	virtual void PostEditChangeProperty(FPropertyChangedEvent& PropertyChangedEvent) override;
+#endif
+	//~ End UObject interface
+	
+
+public:
+	// Called every frame
+	virtual void TickComponent(float DeltaTime, ELevelTick TickType, FActorComponentTickFunction* ThisTickFunction) override;
+
+	void InitiateAnimation();
+
+	void ResetAnimInstance() const;
+	
+private:
+	FTransform CachedLiveLinkTransform;
+	bool bIsDirty;
+	TOptional<bool> bIsSpawnableCache;
+};
diff --git a/Engine/Plugins/VirtualProduction/PerformanceCaptureWorkflow/Source/PerformanceCaptureWorkflowRuntime/Public/PCapStageRoot.h b/Engine/Plugins/VirtualProduction/PerformanceCaptureWorkflow/Source/PerformanceCaptureWorkflowRuntime/Public/PCapStageRoot.h
new file mode 100644
--- /dev/null
+++ b/Engine/Plugins/VirtualProduction/PerformanceCaptureWorkflow/Source/PerformanceCaptureWorkflowRuntime/Public/PCapStageRoot.h
@@ -0,0 +1,37 @@
+// Copyright Epic Games, Inc. All Rights Reserved.
+
+
+#pragma once
+
+#include "CoreMinimal.h"
+#include "GameFramework/Actor.h"
+#include "Components/SceneComponent.h"
+#include "Components/DecalComponent.h"
+#include "Components/SceneCaptureComponent2D.h"
+
+#include "PCapStageRoot.generated.h"
+
+class USceneCaptureComponent2D;
+
+UCLASS(Blueprintable, Abstract)
+class APerformanceCaptureStageRoot : public AActor
+{
+	GENERATED_BODY()
+
+public:
+	// Sets default values for this actor's properties
+	APerformanceCaptureStageRoot(const FObjectInitializer& ObjectInitializer);
+
+	/** Scene capture component for capturing the overhead, orthographic map. */
+	UPROPERTY(EditAnywhere, BlueprintReadWrite,  Category="Peformance Capture|StageRoot")
+	TObjectPtr<USceneCaptureComponent2D> MapCaptureComponent;
+
+	/** Decal component for rendering a grid to the floor. */
+	UPROPERTY(EditAnywhere, BlueprintReadWrite,  Category="Peformance Capture|StageRoot")
+	TObjectPtr<UDecalComponent> DecalComponent;
+
+	/** Scene component, under which all stage ghost meshes should be parented. */
+	UPROPERTY(EditAnywhere, BlueprintReadWrite,  Category="Peformance Capture|StageRoot")
+	TObjectPtr<USceneComponent> StageMeshParent;
+	
+};
\ No newline at end of file
diff --git a/Engine/Plugins/VirtualProduction/PerformanceCaptureWorkflow/Source/PerformanceCaptureWorkflowRuntime/Public/PerformanceCaptureWorkflowRuntime.h b/Engine/Plugins/VirtualProduction/PerformanceCaptureWorkflow/Source/PerformanceCaptureWorkflowRuntime/Public/PerformanceCaptureWorkflowRuntime.h
new file mode 100644
--- /dev/null
+++ b/Engine/Plugins/VirtualProduction/PerformanceCaptureWorkflow/Source/PerformanceCaptureWorkflowRuntime/Public/PerformanceCaptureWorkflowRuntime.h
@@ -0,0 +1,12 @@
+﻿// Copyright Epic Games, Inc. All Rights Reserved.
+#pragma once
+
+#include "CoreMinimal.h"
+#include "Modules/ModuleManager.h"
+
+class FPerformanceCaptureWorkflowRuntimeModule : public IModuleInterface
+{
+public:
+    virtual void StartupModule() override;
+    virtual void ShutdownModule() override;
+};
