[Performance Capture Workflow]

- Introduce the functionality to calculate a dynamic prop offset via BP native function (dynamic constraint)
- Add Function library to query IK Rig and Retargeter properties at run time for calculating dynamic constraint.
- Improved handling of spawning custom props actors from the PCap Prop data asset.

diff --git a/Engine/Plugins/VirtualProduction/PerformanceCaptureWorkflow/Source/PerformanceCaptureWorkflow/Private/PCapActorFactory.cpp b/Engine/Plugins/VirtualProduction/PerformanceCaptureWorkflow/Source/PerformanceCaptureWorkflow/Private/PCapActorFactory.cpp
--- a/Engine/Plugins/VirtualProduction/PerformanceCaptureWorkflow/Source/PerformanceCaptureWorkflow/Private/PCapActorFactory.cpp
+++ b/Engine/Plugins/VirtualProduction/PerformanceCaptureWorkflow/Source/PerformanceCaptureWorkflow/Private/PCapActorFactory.cpp
@@ -370,21 +370,25 @@ void UPCapPropActorFactory::PostSpawnActor(UObject* Asset, AActor* NewActor)
 				FSubobjectDataHandle NewPropComponentHandle = SubSystem->AddNewSubobject(NewObjectParams, FailureReason);
 			}
 			
-			if(NewPropComponent)
+			UE_LOG(LogPCap, Display, TEXT("New PCap Prop Component added to %s"), *NewActor->GetActorLabel())
+		}
+
+		//If the new actor has been spawned from a BP class, check it has a Prop component and set the Live Link Subject
+		//Or, if a new prop component has just been added above set the subject, offset and spawning data asset for later retrieval. 
+		if(NewActor)
+		{
+			NewActor->SetActorLabel(PropDataAsset->PropName.ToString(), false);
+			
+			if(UPCapPropComponent* PropComponent = Cast<UPCapPropComponent>(NewActor->GetComponentByClass(UPCapPropComponent::StaticClass())))
 			{
-				NewPropComponent->SetLiveLinkSubject(PropDataAsset->LiveLinkSubject);
-				NewPropComponent->SetOffsetTransform(PropDataAsset->PropOffsetTransform);
-				if (Asset) 
-				{
-					if (Asset->IsAsset())
-					{
-						FAssetData NewAssetData = FAssetData(Asset);
-						FName SpawningPackageName = NewAssetData.PackageName;
-						NewPropComponent->SpawningDataAsset = SpawningPackageName; //Add the package name of the spawning data asset to the prop component so the toolset can locate the data asset for a given prop
-					}
-				}
+				PropComponent->SetLiveLinkSubject(PropDataAsset->LiveLinkSubject);
+				PropComponent->SetOffsetTransform(PropDataAsset->PropOffsetTransform);
+				PropComponent->SpawningDataAsset = FAssetData(PropDataAsset).PackageName; //Add the package name of the spawning data asset to the prop component so the toolset can locate the data asset for a given prop
+			}
+			else
+			{
+				UE_LOG(LogPCap, Display, TEXT("Prop spawned from data asset with no valid PcapProp component. %s"), *NewActor->GetActorLabel())
 			}
-			UE_LOG(LogPCap, Display, TEXT("New PCap Prop Component added to %s"), *NewActor->GetActorLabel())
 		}
 	}
 }
diff --git a/Engine/Plugins/VirtualProduction/PerformanceCaptureWorkflow/Source/PerformanceCaptureWorkflowRuntime/PerformanceCaptureWorkflowRuntime.Build.cs b/Engine/Plugins/VirtualProduction/PerformanceCaptureWorkflow/Source/PerformanceCaptureWorkflowRuntime/PerformanceCaptureWorkflowRuntime.Build.cs
--- a/Engine/Plugins/VirtualProduction/PerformanceCaptureWorkflow/Source/PerformanceCaptureWorkflowRuntime/PerformanceCaptureWorkflowRuntime.Build.cs
+++ b/Engine/Plugins/VirtualProduction/PerformanceCaptureWorkflow/Source/PerformanceCaptureWorkflowRuntime/PerformanceCaptureWorkflowRuntime.Build.cs
@@ -11,7 +11,7 @@ public class PerformanceCaptureWorkflowRuntime : ModuleRules
         PublicDependencyModuleNames.AddRange(
             new string[]
             {
-                "Core",
+                "Core"
             }
         );
 
@@ -25,6 +25,8 @@ public class PerformanceCaptureWorkflowRuntime : ModuleRules
                 "LiveLink",
                 "LiveLinkInterface",
                 "LiveLinkAnimationCore",
+				"PerformanceCaptureCore", 
+				"IKRig"
             }
         );
     }
diff --git a/Engine/Plugins/VirtualProduction/PerformanceCaptureWorkflow/Source/PerformanceCaptureWorkflowRuntime/Private/AnimNode_LiveLinkProp.cpp b/Engine/Plugins/VirtualProduction/PerformanceCaptureWorkflow/Source/PerformanceCaptureWorkflowRuntime/Private/AnimNode_LiveLinkProp.cpp
--- a/Engine/Plugins/VirtualProduction/PerformanceCaptureWorkflow/Source/PerformanceCaptureWorkflowRuntime/Private/AnimNode_LiveLinkProp.cpp
+++ b/Engine/Plugins/VirtualProduction/PerformanceCaptureWorkflow/Source/PerformanceCaptureWorkflowRuntime/Private/AnimNode_LiveLinkProp.cpp
@@ -35,6 +35,8 @@ void FAnimNode_LiveLinkProp::BuildPoseFromAnimData(const FLiveLinkSubjectFrameDa
 
 		FTransform BoneTransform = OffsetTransform*FrameData->Transforms[i];
 
+		BoneTransform.SetLocation(BoneTransform.GetLocation()+DynamicConstraintOffset);
+
 		int32 MeshIndex = Output.Pose.GetBoneContainer().GetPoseBoneIndexForBoneName(BoneName);
 		if (MeshIndex != INDEX_NONE)
 		{
diff --git a/Engine/Plugins/VirtualProduction/PerformanceCaptureWorkflow/Source/PerformanceCaptureWorkflowRuntime/Private/AnimNode_LiveLinkProp.h b/Engine/Plugins/VirtualProduction/PerformanceCaptureWorkflow/Source/PerformanceCaptureWorkflowRuntime/Private/AnimNode_LiveLinkProp.h
--- a/Engine/Plugins/VirtualProduction/PerformanceCaptureWorkflow/Source/PerformanceCaptureWorkflowRuntime/Private/AnimNode_LiveLinkProp.h
+++ b/Engine/Plugins/VirtualProduction/PerformanceCaptureWorkflow/Source/PerformanceCaptureWorkflowRuntime/Private/AnimNode_LiveLinkProp.h
@@ -36,12 +36,14 @@ public:
 	UPROPERTY(EditAnywhere, Transient, Category = "Performance Capture|Prop")
 	bool bDoLiveLinkEvaluation = true;
 
-	/* Transform to apply local space to offset to the incoming data. */
+	/* Transform to apply local space to offset to the incoming Live Link data. */
 	UPROPERTY(EditAnywhere, Category = "Performance Capture|Prop")
 	FTransform OffsetTransform = FTransform::Identity;
 
+	/* Location to apply in parent bone space to offset the incoming Live Link data. */
+	UPROPERTY(EditAnywhere, Category = "Performance Capture|Prop")
+	FVector DynamicConstraintOffset = FVector::ZeroVector;
 	
-
 	FAnimNode_LiveLinkProp();
 
 	//~ FAnimNode_Base interface
diff --git a/Engine/Plugins/VirtualProduction/PerformanceCaptureWorkflow/Source/PerformanceCaptureWorkflowRuntime/Private/PCapPropComponent.cpp b/Engine/Plugins/VirtualProduction/PerformanceCaptureWorkflow/Source/PerformanceCaptureWorkflowRuntime/Private/PCapPropComponent.cpp
--- a/Engine/Plugins/VirtualProduction/PerformanceCaptureWorkflow/Source/PerformanceCaptureWorkflowRuntime/Private/PCapPropComponent.cpp
+++ b/Engine/Plugins/VirtualProduction/PerformanceCaptureWorkflow/Source/PerformanceCaptureWorkflowRuntime/Private/PCapPropComponent.cpp
@@ -2,12 +2,17 @@
 
 
 #include "PCapPropComponent.h"
+
+#include "CaptureCharacter.h"
 #include "Components/SkeletalMeshComponent.h"
 #include "ILiveLinkClient.h"
 #include "PCapPropLiveLinkAnimInstance.h"
+#include "PerformanceCaptureWorkflowRuntime.h"
+#include "RetargetComponent.h"
 #include "GameFramework/Actor.h"
 #include "Components/SceneComponent.h"
 #include "Features/IModularFeatures.h"
+#include "Retargeter/IKRetargeter.h"
 #include "Roles/LiveLinkAnimationRole.h"
 #include "Roles/LiveLinkTransformRole.h"
 #include "Roles/LiveLinkTransformTypes.h"
@@ -68,8 +73,8 @@ void UPCapPropComponent::SetEvaluateLiveLinkData(bool bEvaluate)
 	USkeletalMeshComponent* ControlledMesh = Cast<USkeletalMeshComponent>(ControlledComponent.GetComponent(GetOwner()));
 	if(IsValid(ControlledMesh))
 	{
-		UPCapPropLiveLinkAnimInstance* LiveLinkAnimInstance;
-		LiveLinkAnimInstance = Cast<UPCapPropLiveLinkAnimInstance>(ControlledMesh->GetAnimInstance());
+		UPCapPropLiveLinkAnimInstance* LiveLinkAnimInstance = Cast<UPCapPropLiveLinkAnimInstance>(
+			ControlledMesh->GetAnimInstance());
 		if(LiveLinkAnimInstance)
 		{
 			LiveLinkAnimInstance->EnableLiveLinkEvaluation(bEvaluateLiveLink);
@@ -90,15 +95,58 @@ void UPCapPropComponent::SetOffsetTransform(FTransform NewOffset)
 	TObjectPtr<USkeletalMeshComponent> ControlledMesh = Cast<USkeletalMeshComponent>(ControlledComponent.GetComponent(GetOwner()));
 	if(IsValid(ControlledMesh))
 	{
-		UPCapPropLiveLinkAnimInstance* LiveLinkAnimInstance;
-		LiveLinkAnimInstance = Cast<UPCapPropLiveLinkAnimInstance>(ControlledMesh->GetAnimInstance());
-		if(LiveLinkAnimInstance)
+		if(UPCapPropLiveLinkAnimInstance* LiveLinkAnimInstance = Cast<UPCapPropLiveLinkAnimInstance>(ControlledMesh->GetAnimInstance()))
 		{
 			LiveLinkAnimInstance->SetOffsetTransform(OffsetTransform);
 		}
 	}
 }
 
+TArray<FName> UPCapPropComponent::GetAttachBones() const
+{
+	TArray<FName> FoundIKChainBones;
+	
+	for (TSoftObjectPtr<ACaptureCharacter> Character : DynamicAttachmentCharacters)
+	{
+		if (Character.IsValid())
+		{
+			if (const UIKRetargeter* Retargeter = Character.Get()->GetRetargetComponent()->RetargetAsset)
+			{
+				if (const UIKRigDefinition* TargetIKRig = Retargeter->GetIKRig(ERetargetSourceOrTarget::Target))
+				{
+					TArray<FName> ChainNames = TargetIKRig->GetRetargetChainNames();
+
+					for (FName ChainName : ChainNames)
+					{
+						if (bAllowOnlyIKBones)
+						{
+							if (const FBoneChain* Chain = TargetIKRig->GetRetargetChainByName(ChainName); Chain->IKGoalName != NAME_None)
+							{
+								FoundIKChainBones.AddUnique(Chain->StartBone.BoneName);
+								FoundIKChainBones.AddUnique(Chain->EndBone.BoneName);
+							}
+						}
+						else
+						{
+							const FBoneChain* Chain = TargetIKRig->GetRetargetChainByName(ChainName);
+							FoundIKChainBones.AddUnique(Chain->StartBone.BoneName);
+							FoundIKChainBones.AddUnique(Chain->EndBone.BoneName);
+						}
+					}
+				}
+			}
+				
+		}
+	}
+	return FoundIKChainBones;
+}
+
+FVector UPCapPropComponent::CalculateDynamicOffset_Implementation()
+{
+	FVector ReturnValue = FVector::ZeroVector;
+	return ReturnValue;
+}
+
 void UPCapPropComponent::OnRegister()
 {
 	Super::OnRegister();
@@ -154,7 +202,9 @@ void UPCapPropComponent::PostEditChangeProperty(FPropertyChangedEvent& PropertyC
 
 void UPCapPropComponent::TickComponent(float DeltaTime, ELevelTick TickType, FActorComponentTickFunction* ThisTickFunction)
 {
-
+	// Calculate the dynamic constraint. Will call to the BP implmented method if there is one available. 
+	DynamicConstraintOffset = CalculateDynamicOffset();
+	
 	//Check if spawnable
 	if(!bIsSpawnableCache.IsSet() || bIsDirty)
 	{
@@ -188,8 +238,19 @@ void UPCapPropComponent::TickComponent(float DeltaTime, ELevelTick TickType, FAc
 				if(FrameData && bEvaluateLiveLink)
 				{
 					CachedLiveLinkTransform = FrameData->Transforms[0];
+					FinalTransform = CachedLiveLinkTransform;
 				}
-				GetControlledComponent()->SetRelativeTransform(OffsetTransform * CachedLiveLinkTransform, false, nullptr, ETeleportType::None);
+
+				if (bUseDynamicConstraint)
+				{
+					FinalTransform.SetComponents(CachedLiveLinkTransform.GetRotation(),CachedLiveLinkTransform.GetLocation()+DynamicConstraintOffset, CachedLiveLinkTransform.GetScale3D());
+				}
+				else
+				{
+					FinalTransform.SetComponents(CachedLiveLinkTransform.GetRotation(),CachedLiveLinkTransform.GetLocation(), CachedLiveLinkTransform.GetScale3D());
+				}
+				
+				GetControlledComponent()->SetRelativeTransform(OffsetTransform * FinalTransform, false, nullptr, ETeleportType::None);
 
 			}
 			else if(LiveLinkClient.GetSubjectRole_AnyThread(SubjectName) == ULiveLinkTransformRole::StaticClass())
@@ -199,7 +260,17 @@ void UPCapPropComponent::TickComponent(float DeltaTime, ELevelTick TickType, FAc
 				{
 					CachedLiveLinkTransform = FrameData->Transform;
 				}
-				GetControlledComponent()->SetRelativeTransform(OffsetTransform * CachedLiveLinkTransform, false, nullptr, ETeleportType::None);
+				
+				if (bUseDynamicConstraint)
+				{
+					FinalTransform.SetComponents(CachedLiveLinkTransform.GetRotation(),CachedLiveLinkTransform.GetLocation()+DynamicConstraintOffset, CachedLiveLinkTransform.GetScale3D());
+				}
+				else
+				{
+					FinalTransform.SetComponents(CachedLiveLinkTransform.GetRotation(),CachedLiveLinkTransform.GetLocation(), CachedLiveLinkTransform.GetScale3D());
+				}
+				
+				GetControlledComponent()->SetRelativeTransform(OffsetTransform * FinalTransform, false, nullptr, ETeleportType::None);
 			}
 		}
 		//If we are controlling a Skelmesh and receiving a transform role, apply this route. 
@@ -218,12 +289,31 @@ void UPCapPropComponent::TickComponent(float DeltaTime, ELevelTick TickType, FAc
 				{
 					CachedLiveLinkTransform = FrameData->Transform;
 				}
-				GetControlledComponent()->SetRelativeTransform(OffsetTransform * CachedLiveLinkTransform, false, nullptr, ETeleportType::None);
+
+				if (bUseDynamicConstraint)
+				{
+					FinalTransform.SetComponents(CachedLiveLinkTransform.GetRotation(),CachedLiveLinkTransform.GetLocation()+DynamicConstraintOffset, CachedLiveLinkTransform.GetScale3D());
+				}
+				else
+				{
+					FinalTransform.SetComponents(CachedLiveLinkTransform.GetRotation(),CachedLiveLinkTransform.GetLocation(), CachedLiveLinkTransform.GetScale3D());
+				}
+				
+				GetControlledComponent()->SetRelativeTransform(OffsetTransform * FinalTransform, false, nullptr, ETeleportType::None);
 			}
+			
 			if((SkelLiveLinkClient.GetSubjectRole_AnyThread(SubjectName) == ULiveLinkAnimationRole::StaticClass()) && bIsDirty)
 			{
 				InitiateAnimation();
 			}
+			
+			if (bUseDynamicConstraint && ControlledSkelmesh)
+			{
+				if(UPCapPropLiveLinkAnimInstance* AnimInstance = Cast<UPCapPropLiveLinkAnimInstance>(ControlledSkelmesh->GetAnimInstance()))
+				{
+					AnimInstance->SetDynamicConstraintVector(DynamicConstraintOffset); // This is only ever applied in the space of the root bone of the skeleton.
+				}
+			}
 		}
 	}
 	Super::TickComponent(DeltaTime, TickType, ThisTickFunction);
@@ -244,8 +334,7 @@ void UPCapPropComponent::InitiateAnimation()
 		ControlledMesh->VisibilityBasedAnimTickOption = EVisibilityBasedAnimTickOption::AlwaysTickPoseAndRefreshBones;
 		if(ControlledMesh)
 		{
-			UPCapPropLiveLinkAnimInstance* AnimInstance;
-			AnimInstance = Cast<UPCapPropLiveLinkAnimInstance>(ControlledMesh->GetAnimInstance());
+			UPCapPropLiveLinkAnimInstance* AnimInstance = Cast<UPCapPropLiveLinkAnimInstance>(ControlledMesh->GetAnimInstance());
 			if(AnimInstance)
 			{
 				AnimInstance->SetSubject(SubjectName);
diff --git a/Engine/Plugins/VirtualProduction/PerformanceCaptureWorkflow/Source/PerformanceCaptureWorkflowRuntime/Private/PCapPropLiveLinkAnimInstance.h b/Engine/Plugins/VirtualProduction/PerformanceCaptureWorkflow/Source/PerformanceCaptureWorkflowRuntime/Private/PCapPropLiveLinkAnimInstance.h
--- a/Engine/Plugins/VirtualProduction/PerformanceCaptureWorkflow/Source/PerformanceCaptureWorkflowRuntime/Private/PCapPropLiveLinkAnimInstance.h
+++ b/Engine/Plugins/VirtualProduction/PerformanceCaptureWorkflow/Source/PerformanceCaptureWorkflowRuntime/Private/PCapPropLiveLinkAnimInstance.h
@@ -60,6 +60,12 @@ public:
 	{
 		GetProxyOnGameThread<FPCapPropLiveLinkAnimInstanceProxy>().PoseNode.OffsetTransform = Offset;
 	}
+
+	UFUNCTION(BlueprintCallable, Category = "Performance Capture|Prop")
+	void SetDynamicConstraintVector(FVector DynamicOffset)
+	{
+		GetProxyOnGameThread<FPCapPropLiveLinkAnimInstanceProxy>().PoseNode.DynamicConstraintOffset = DynamicOffset;
+	}
 	
 	//Not exposed to Blueprint
 	bool GetEnableLiveLinkEvaluation()
diff --git a/Engine/Plugins/VirtualProduction/PerformanceCaptureWorkflow/Source/PerformanceCaptureWorkflowRuntime/Private/PCapWorkflowRuntimeFunctionLibrary.cpp b/Engine/Plugins/VirtualProduction/PerformanceCaptureWorkflow/Source/PerformanceCaptureWorkflowRuntime/Private/PCapWorkflowRuntimeFunctionLibrary.cpp
new file mode 100644
--- /dev/null
+++ b/Engine/Plugins/VirtualProduction/PerformanceCaptureWorkflow/Source/PerformanceCaptureWorkflowRuntime/Private/PCapWorkflowRuntimeFunctionLibrary.cpp
@@ -0,0 +1,105 @@
+﻿// Copyright Epic Games, Inc. All Rights Reserved.
+
+
+#include "PCapWorkflowRuntimeFunctionLibrary.h"
+
+#include "PerformanceCaptureWorkflowRuntime.h"
+#include "Retargeter/IKRetargeter.h"
+#include "Retargeter/RetargetOps/FKChainsOp.h"
+
+const UIKRigDefinition* UPCapWorkflowRuntimeFunctionLibrary::GetSourceRig(const UIKRetargeter* InRetargetAsset)
+{
+	if (InRetargetAsset)
+	{
+		return InRetargetAsset->GetIKRig(ERetargetSourceOrTarget::Source);
+	}
+	return nullptr;
+}
+
+const UIKRigDefinition* UPCapWorkflowRuntimeFunctionLibrary::GetTargetRig(const UIKRetargeter* InRetargetAsset)
+{
+	if (InRetargetAsset)
+	{
+		return InRetargetAsset->GetIKRig(ERetargetSourceOrTarget::Target);
+	}
+	return nullptr;
+}
+
+TArray<FBoneChain> UPCapWorkflowRuntimeFunctionLibrary::GetRetargetChains(const UIKRigDefinition* InIKRig)
+{
+	if (InIKRig)
+	{
+		return InIKRig->GetRetargetChains();
+	}
+	return TArray<FBoneChain>();
+}
+
+FName UPCapWorkflowRuntimeFunctionLibrary::GetChainStartBone(const UIKRigDefinition* InIKRig, const FName ChainName)
+{
+	FName FoundBone = NAME_None;
+	if (InIKRig)
+	{
+		const FBoneChain* Chain = InIKRig->GetRetargetChainByName(ChainName);
+		if (!Chain)
+		{
+			UE_LOG(LogPCapRuntime, Verbose, TEXT("Retarget chain not found: %s."), *ChainName.ToString())
+			return NAME_None;
+		}
+		FoundBone = Chain->StartBone.BoneName;
+	}
+	return FoundBone;
+}
+
+FName UPCapWorkflowRuntimeFunctionLibrary::GetChainEndBone(const UIKRigDefinition* InIKRig, const FName ChainName)
+{
+	FName FoundBone = NAME_None;
+	if (InIKRig)
+	{
+		const FBoneChain* Chain = InIKRig->GetRetargetChainByName(ChainName);
+		if (!Chain)
+		{
+			UE_LOG(LogPCapRuntime, Verbose, TEXT("Retarget chain not found: %s."), *ChainName.ToString())
+			return NAME_None;
+		}
+		FoundBone = Chain->EndBone.BoneName;
+	}
+	return FoundBone;
+}
+
+FName UPCapWorkflowRuntimeFunctionLibrary::GetChainFromBone(const UIKRigDefinition* InIKRig, const FName InBoneName)
+{
+	FName ChainName = NAME_None;
+	
+	if (InIKRig)
+	{
+		TArray<FBoneChain> Chains = InIKRig->GetRetargetChains();
+
+		for (FBoneChain Chain : Chains)
+		{
+			if (InBoneName == Chain.StartBone.BoneName || InBoneName == Chain.EndBone.BoneName)
+			{
+				ChainName = Chain.ChainName;
+			}
+		}
+	}
+	return ChainName;
+}
+
+FName UPCapWorkflowRuntimeFunctionLibrary::GetChainPair(UIKRetargeter* InRetargetAsset, const FName InChainName,
+                                                        const ERetargetSourceOrTarget SourceOrTarget)
+{
+	FName ChainName = NAME_None;
+
+	if (InRetargetAsset && InChainName!=NAME_None)
+	{
+		if (const FIKRetargetFKChainsOp* FKOp = InRetargetAsset->GetFirstRetargetOpOfType<FIKRetargetFKChainsOp>())
+		{
+			const FIKRetargetFKChainsOpSettings& ChainsSettings = FKOp->Settings;
+			if (const FRetargetChainPair* ChainPair = ChainsSettings.ChainMapping.FindChainPair(InChainName,SourceOrTarget))
+			{
+				ChainName = ChainPair->SourceChainName;
+			}
+		}
+	}
+	return ChainName;
+}
diff --git a/Engine/Plugins/VirtualProduction/PerformanceCaptureWorkflow/Source/PerformanceCaptureWorkflowRuntime/Private/PCapWorkflowRuntimeFunctionLibrary.h b/Engine/Plugins/VirtualProduction/PerformanceCaptureWorkflow/Source/PerformanceCaptureWorkflowRuntime/Private/PCapWorkflowRuntimeFunctionLibrary.h
new file mode 100644
--- /dev/null
+++ b/Engine/Plugins/VirtualProduction/PerformanceCaptureWorkflow/Source/PerformanceCaptureWorkflowRuntime/Private/PCapWorkflowRuntimeFunctionLibrary.h
@@ -0,0 +1,58 @@
+﻿// Copyright Epic Games, Inc. All Rights Reserved.
+
+#pragma once
+
+#include "CoreMinimal.h"
+#include "Kismet/BlueprintFunctionLibrary.h"
+#include "PCapWorkflowRuntimeFunctionLibrary.generated.h"
+
+/**
+ * Function for Performance Capture worfklow, enabling read-only access to Retargeter properties.
+ * Intended for querying retargeters for calculation of dynamic prop constraints. 
+ */
+UCLASS()
+class PERFORMANCECAPTUREWORKFLOWRUNTIME_API UPCapWorkflowRuntimeFunctionLibrary : public UBlueprintFunctionLibrary
+{
+	GENERATED_BODY()
+	/**
+	 * Get the source IK rig from a IKRetargeter asset.
+	 */
+	UFUNCTION(BlueprintCallable, Category="Performance Capture|Runtime")
+	static const UIKRigDefinition* GetSourceRig(const UIKRetargeter* InRetargetAsset);
+
+	/**
+	 * Get the target IK rig from a IKRetargeter asset.
+	 */
+	UFUNCTION(BlueprintCallable, Category="Performance Capture|Runtime")
+	static const UIKRigDefinition* GetTargetRig(const UIKRetargeter* InRetargetAsset);
+
+	/**
+	 * Get all the chains in the given IKRig Asset.
+	 */
+	UFUNCTION(BlueprintCallable, Category="Performance Capture|Runtime")
+	static TArray<FBoneChain> GetRetargetChains(const UIKRigDefinition* InIKRig);
+
+	/**
+	 * Returns the bone name at the start of a given chain.
+	 */
+	UFUNCTION(BlueprintCallable, Category="Performance Capture|Runtime")
+	static FName GetChainStartBone(const UIKRigDefinition* InIKRig, const FName ChainName);
+
+	/**
+	 * Returns the bone name at the end of a given chain.
+	 */
+	UFUNCTION(BlueprintCallable, Category="Performance Capture|Runtime")
+	static FName GetChainEndBone(const UIKRigDefinition* InIKRig, const FName ChainName);
+
+	/**
+	 * Returns the first chain the given bone name is a member of.
+	 */
+	UFUNCTION(BlueprintCallable, Category="Performance Capture|Runtime")
+	static FName GetChainFromBone(const UIKRigDefinition* InIKRig, const FName InBoneName);
+
+	/**
+	 * Get the corresponding source or target chain from a given chain name.
+	 */
+	UFUNCTION(BlueprintCallable, Category="Performance Capture|Runtime")
+	static FName GetChainPair(UIKRetargeter* InRetargetAsset, const FName InChainName, const ERetargetSourceOrTarget SourceOrTarget);
+};
diff --git a/Engine/Plugins/VirtualProduction/PerformanceCaptureWorkflow/Source/PerformanceCaptureWorkflowRuntime/Private/PerformanceCaptureWorkflowRuntime.cpp b/Engine/Plugins/VirtualProduction/PerformanceCaptureWorkflow/Source/PerformanceCaptureWorkflowRuntime/Private/PerformanceCaptureWorkflowRuntime.cpp
--- a/Engine/Plugins/VirtualProduction/PerformanceCaptureWorkflow/Source/PerformanceCaptureWorkflowRuntime/Private/PerformanceCaptureWorkflowRuntime.cpp
+++ b/Engine/Plugins/VirtualProduction/PerformanceCaptureWorkflow/Source/PerformanceCaptureWorkflowRuntime/Private/PerformanceCaptureWorkflowRuntime.cpp
@@ -3,9 +3,11 @@
 
 #define LOCTEXT_NAMESPACE "FPerformanceCaptureWorkflowRuntimeModule"
 
+DEFINE_LOG_CATEGORY(LogPCapRuntime);
+
 void FPerformanceCaptureWorkflowRuntimeModule::StartupModule()
 {
-    
+	
 }
 
 void FPerformanceCaptureWorkflowRuntimeModule::ShutdownModule()
diff --git a/Engine/Plugins/VirtualProduction/PerformanceCaptureWorkflow/Source/PerformanceCaptureWorkflowRuntime/Public/PCapPropComponent.h b/Engine/Plugins/VirtualProduction/PerformanceCaptureWorkflow/Source/PerformanceCaptureWorkflowRuntime/Public/PCapPropComponent.h
--- a/Engine/Plugins/VirtualProduction/PerformanceCaptureWorkflow/Source/PerformanceCaptureWorkflowRuntime/Public/PCapPropComponent.h
+++ b/Engine/Plugins/VirtualProduction/PerformanceCaptureWorkflow/Source/PerformanceCaptureWorkflowRuntime/Public/PCapPropComponent.h
@@ -9,6 +9,7 @@
 #include "Components/StaticMeshComponent.h"
 #include "PCapPropComponent.generated.h"
 
+class ACaptureCharacter;
 /**
 * A component for use in Motion Capture props. Accepts data as either a LiveLink Transform Role or an Animation Role. If the Live Link data is in the Animation role, this component will take the root bone transform and apply that transform to the owning actor's root component. Can be applied to Static Meshes, Skeletal Meshes and Blueprint constructions. Transform data can be offset in the component's local space. If this component is driving a Skeletal Mesh component and is receiving Live Link Animation data the full bone hierarchy will be applied. 
 */
@@ -44,7 +45,7 @@ public:
 	/**
 	 *The package name of the spawning data asset, if there is one. 
 	 */
-	UPROPERTY(BlueprintReadOnly, Category = "Performance Capture")
+	UPROPERTY(BlueprintReadOnly, VisibleAnywhere, Category = "Performance Capture")
 	FName SpawningDataAsset;
 	
 	/**
@@ -98,7 +99,7 @@ public:
 	*/
 	UFUNCTION(BlueprintPure, Category ="Performance Capture|Prop")
 	bool GetEvaluateLiveLinkData();
-
+	
 	/**
 	* Update the local space offset to this prop's Live Link pose. 
 	* @param NewOffset Offset Transform.
@@ -106,7 +107,46 @@ public:
 	UFUNCTION(BlueprintCallable, Category ="Performance Capture|Prop")
 	void SetOffsetTransform(FTransform NewOffset);
 
+	/**
+	* Enable/Disable use of the Dynamic Constraint offset vector. 
+	*/
+	UPROPERTY(BlueprintReadWrite, EditAnywhere, Category ="Performance Capture|Dynamic Constraint")
+	bool bUseDynamicConstraint = false;
+	
+	/**
+	* List of Capture Characters to consider in calculating a dynamic constraint.
+	*/
+	UPROPERTY(EditAnywhere, BlueprintReadOnly, Category = "Performance Capture|Dynamic Constraint", meta = (EditCondition="bUseDynamicConstraint", EditConditionHides))
+	TArray<TSoftObjectPtr<ACaptureCharacter>> DynamicAttachmentCharacters;
+
+	/**
+	* False = All bone chains' start and end bones will be considered for dynamic attachment .
+	* True  = Only bones in chains with an IK goal will be considered for dynamic attachment.
+	*/
+	UPROPERTY(EditAnywhere, BlueprintReadOnly, Category = "Performance Capture|Dynamic Constraint", meta = (EditCondition="bUseDynamicConstraint", EditConditionHides))
+	bool bAllowOnlyIKBones = true;
 
+	/**
+	* A list of bones of to test against for dynamic constraint.
+	*/
+	UPROPERTY(EditAnywhere, BlueprintReadOnly, Category = "Performance Capture|Dynamic Constraint", meta = (GetOptions="GetAttachBones", EditCondition="bUseDynamicConstraint", EditConditionHides))
+	TArray<FName> DynamicConstraintAttachBones;
+
+	/**
+	 * Return a list of bones from the given list of DynamicAttachmentCharacters.
+	 * @return List of bones to use in DynamicConstraintAttachBones variable
+	 */
+	UFUNCTION()
+	TArray<FName> GetAttachBones() const;
+
+	/**
+	 * Blueprint implementable event to enable calculation of a location offset before applying it to
+	 * the controlled component's transform.
+	 * @return World-space location offset to apply to the prop's transform
+	 */
+	UFUNCTION(BlueprintNativeEvent, Category = "Performance Capture|Dynamic Constraint")
+	FVector CalculateDynamicOffset();
+	
 protected:
 
 	//~ Begin UActorComponent interface
@@ -130,7 +170,13 @@ public:
 	void ResetAnimInstance() const;
 	
 private:
-	FTransform CachedLiveLinkTransform;
+	FVector DynamicConstraintOffset = FVector::ZeroVector;
+	
+	FTransform CachedLiveLinkTransform = FTransform::Identity;
+	
+	FTransform FinalTransform = FTransform::Identity;
+
 	bool bIsDirty;
+
 	TOptional<bool> bIsSpawnableCache;
 };
diff --git a/Engine/Plugins/VirtualProduction/PerformanceCaptureWorkflow/Source/PerformanceCaptureWorkflowRuntime/Public/PerformanceCaptureWorkflowRuntime.h b/Engine/Plugins/VirtualProduction/PerformanceCaptureWorkflow/Source/PerformanceCaptureWorkflowRuntime/Public/PerformanceCaptureWorkflowRuntime.h
--- a/Engine/Plugins/VirtualProduction/PerformanceCaptureWorkflow/Source/PerformanceCaptureWorkflowRuntime/Public/PerformanceCaptureWorkflowRuntime.h
+++ b/Engine/Plugins/VirtualProduction/PerformanceCaptureWorkflow/Source/PerformanceCaptureWorkflowRuntime/Public/PerformanceCaptureWorkflowRuntime.h
@@ -4,6 +4,8 @@
 #include "CoreMinimal.h"
 #include "Modules/ModuleManager.h"
 
+DECLARE_LOG_CATEGORY_EXTERN(LogPCapRuntime, Log, All);
+
 class FPerformanceCaptureWorkflowRuntimeModule : public IModuleInterface
 {
 public:
