[Performance Capture workflow]

-CIS warning fixes. Bughawk documenttion and PVS incorrect init order.

#rnx

diff --git a/Engine/Plugins/VirtualProduction/PerformanceCaptureWorkflow/Source/PerformanceCaptureWorkflow/Private/PCapActorFactory.cpp b/Engine/Plugins/VirtualProduction/PerformanceCaptureWorkflow/Source/PerformanceCaptureWorkflow/Private/PCapActorFactory.cpp
--- a/Engine/Plugins/VirtualProduction/PerformanceCaptureWorkflow/Source/PerformanceCaptureWorkflow/Private/PCapActorFactory.cpp
+++ b/Engine/Plugins/VirtualProduction/PerformanceCaptureWorkflow/Source/PerformanceCaptureWorkflow/Private/PCapActorFactory.cpp
@@ -308,75 +308,76 @@ bool UPCapPropActorFactory::CanCreateActorFrom(const FAssetData& AssetData, FTex
 
 void UPCapPropActorFactory::PostSpawnActor(UObject* Asset, AActor* NewActor)
 {
-	Super::PostSpawnActor(Asset, NewActor);	
+	Super::PostSpawnActor(Asset, NewActor);
 	
-	if(UPCapPropDataAsset* PropDataAsset = Cast<UPCapPropDataAsset>(Asset))
+	if(NewActor)
 	{
-		if(AStaticMeshActor* NewStaticMeshProp = Cast<AStaticMeshActor>(NewActor))
-		{		
-			if(UStaticMesh* StaticMesh = PropDataAsset->PropStaticMesh.LoadSynchronous())
-			{
-				NewStaticMeshProp->GetStaticMeshComponent()->SetStaticMesh(PropDataAsset->PropStaticMesh.LoadSynchronous());
-				NewStaticMeshProp->GetStaticMeshComponent()->SetReceivesDecals(false);
-				if(PropDataAsset->PropName.IsValid()) //If the user has filled out a valid FName use that for the actor label
-				{
-					NewStaticMeshProp->SetActorLabel(PropDataAsset->PropName.ToString(), false);
-				}
-				else //Otherwise use the name of the DataAsset
+		if(UPCapPropDataAsset* PropDataAsset = Cast<UPCapPropDataAsset>(Asset))
+		{
+			if(AStaticMeshActor* NewStaticMeshProp = Cast<AStaticMeshActor>(NewActor))
+			{		
+				if(UStaticMesh* StaticMesh = PropDataAsset->PropStaticMesh.LoadSynchronous())
 				{
-					NewStaticMeshProp->SetActorLabel(PropDataAsset->GetName(), false);
+					NewStaticMeshProp->GetStaticMeshComponent()->SetStaticMesh(PropDataAsset->PropStaticMesh.LoadSynchronous());
+					NewStaticMeshProp->GetStaticMeshComponent()->SetReceivesDecals(false);
+					if(PropDataAsset->PropName.IsValid()) //If the user has filled out a valid FName use that for the actor label
+					{
+						NewStaticMeshProp->SetActorLabel(PropDataAsset->PropName.ToString(), false);
+					}
+					else //Otherwise use the name of the DataAsset
+					{
+						NewStaticMeshProp->SetActorLabel(PropDataAsset->GetName(), false);
+					}
 				}
+				NewStaticMeshProp->SetActorHiddenInGame(PropDataAsset->bHiddenInGame);
 			}
-			NewStaticMeshProp->SetActorHiddenInGame(PropDataAsset->bHiddenInGame);
-		}
-		
-		if(ASkeletalMeshActor* NewSkeletalMeshProp = Cast<ASkeletalMeshActor>(NewActor))
-		{
-			if(USkeletalMesh* SkeletalMesh = PropDataAsset->PropSkeletalMesh.LoadSynchronous())
+			
+			if(ASkeletalMeshActor* NewSkeletalMeshProp = Cast<ASkeletalMeshActor>(NewActor))
 			{
-				NewSkeletalMeshProp->GetSkeletalMeshComponent()->SetSkeletalMeshAsset(PropDataAsset->PropSkeletalMesh.LoadSynchronous());
-				NewSkeletalMeshProp->GetSkeletalMeshComponent()->SetReceivesDecals(false);
-				if(PropDataAsset->PropName.IsValid()) //If the user has filled out a valid FName use that for the actor label
-				{
-					NewSkeletalMeshProp->SetActorLabel(PropDataAsset->PropName.ToString(), false);
-				}
-				else //Otherwise use the name of the DataAsset
+				if(USkeletalMesh* SkeletalMesh = PropDataAsset->PropSkeletalMesh.LoadSynchronous())
 				{
-					NewSkeletalMeshProp->SetActorLabel(PropDataAsset->GetName(), false);
+					NewSkeletalMeshProp->GetSkeletalMeshComponent()->SetSkeletalMeshAsset(PropDataAsset->PropSkeletalMesh.LoadSynchronous());
+					NewSkeletalMeshProp->GetSkeletalMeshComponent()->SetReceivesDecals(false);
+					if(PropDataAsset->PropName.IsValid()) //If the user has filled out a valid FName use that for the actor label
+					{
+						NewSkeletalMeshProp->SetActorLabel(PropDataAsset->PropName.ToString(), false);
+					}
+					else //Otherwise use the name of the DataAsset
+					{
+						NewSkeletalMeshProp->SetActorLabel(PropDataAsset->GetName(), false);
+					}
+					
 				}
-				
+				NewSkeletalMeshProp->SetActorHiddenInGame(PropDataAsset->bHiddenInGame);
 			}
-			NewSkeletalMeshProp->SetActorHiddenInGame(PropDataAsset->bHiddenInGame);
-		}
-		//Only add Subcomponents for the non-transient actor because transient means it is still being placed. Also, don't add the prop component for custom classes
-		if(NewActor->HasAnyFlags(RF_Transient) != true && NewActor->GetClass()!= (PropDataAsset->CustomPropClass.Get())) 
-		{
-			USubobjectDataSubsystem* SubSystem = GEngine->GetEngineSubsystem<USubobjectDataSubsystem>();
-			TObjectPtr<UPCapPropComponent> NewPropComponent =  Cast<UPCapPropComponent>(NewActor->GetComponentByClass(UPCapPropComponent::StaticClass()));
-			if(!NewPropComponent)
+			//Only add Subcomponents for the non-transient actor because transient means it is still being placed. Also, don't add the prop component for custom classes
+			if(NewActor->HasAnyFlags(RF_Transient) != true && NewActor->GetClass()!= (PropDataAsset->CustomPropClass.Get())) 
 			{
-				TArray<FSubobjectDataHandle> Handles;
-				SubSystem->GatherSubobjectData(NewActor, Handles);
+				USubobjectDataSubsystem* SubSystem = GEngine->GetEngineSubsystem<USubobjectDataSubsystem>();
+				TObjectPtr<UPCapPropComponent> NewPropComponent =  Cast<UPCapPropComponent>(NewActor->GetComponentByClass(UPCapPropComponent::StaticClass()));
+				if(!NewPropComponent)
+				{
+					TArray<FSubobjectDataHandle> Handles;
+					SubSystem->GatherSubobjectData(NewActor, Handles);
 
-				FSubobjectDataHandle Handle = Handles[0]; //Root Component is always first entry in handles array
+					FSubobjectDataHandle Handle = Handles[0]; //Root Component is always first entry in handles array
 
-				FAddNewSubobjectParams NewObjectParams;
+					FAddNewSubobjectParams NewObjectParams;
 
-				NewObjectParams.ParentHandle = Handle;
-				NewObjectParams.NewClass = UPCapPropComponent::StaticClass();
-				NewObjectParams.bConformTransformToParent = true;
+					NewObjectParams.ParentHandle = Handle;
+					NewObjectParams.NewClass = UPCapPropComponent::StaticClass();
+					NewObjectParams.bConformTransformToParent = true;
 
-				FText FailureReason;
-				FSubobjectDataHandle NewPropComponentHandle = SubSystem->AddNewSubobject(NewObjectParams, FailureReason);
+					FText FailureReason;
+					FSubobjectDataHandle NewPropComponentHandle = SubSystem->AddNewSubobject(NewObjectParams, FailureReason);
+				}
+				
+				UE_LOG(LogPCap, Display, TEXT("New PCap Prop Component added to %s"), *NewActor->GetActorLabel())
 			}
-			
-			UE_LOG(LogPCap, Display, TEXT("New PCap Prop Component added to %s"), *NewActor->GetActorLabel())
-		}
 
-		//If the new actor has been spawned from a BP class, check it has a Prop component and set the Live Link Subject
-		//Or, if a new prop component has just been added above set the subject, offset and spawning data asset for later retrieval. 
-		if(NewActor)
-		{
+			//If the new actor has been spawned from a BP class, check it has a Prop component and set the Live Link Subject
+			//Or, if a new prop component has just been added above set the subject, offset and spawning data asset for later retrieval. 
+			
 			NewActor->SetActorLabel(PropDataAsset->PropName.ToString(), false);
 			
 			if(UPCapPropComponent* PropComponent = Cast<UPCapPropComponent>(NewActor->GetComponentByClass(UPCapPropComponent::StaticClass())))
diff --git a/Engine/Plugins/VirtualProduction/PerformanceCaptureWorkflow/Source/PerformanceCaptureWorkflowRuntime/Private/PCapWorkflowRuntimeFunctionLibrary.cpp b/Engine/Plugins/VirtualProduction/PerformanceCaptureWorkflow/Source/PerformanceCaptureWorkflowRuntime/Private/PCapWorkflowRuntimeFunctionLibrary.cpp
--- a/Engine/Plugins/VirtualProduction/PerformanceCaptureWorkflow/Source/PerformanceCaptureWorkflowRuntime/Private/PCapWorkflowRuntimeFunctionLibrary.cpp
+++ b/Engine/Plugins/VirtualProduction/PerformanceCaptureWorkflow/Source/PerformanceCaptureWorkflowRuntime/Private/PCapWorkflowRuntimeFunctionLibrary.cpp
@@ -79,6 +79,7 @@ FName UPCapWorkflowRuntimeFunctionLibrary::GetChainFromBone(const UIKRigDefiniti
 			if (InBoneName == Chain.StartBone.BoneName || InBoneName == Chain.EndBone.BoneName)
 			{
 				ChainName = Chain.ChainName;
+				return ChainName;
 			}
 		}
 	}
diff --git a/Engine/Plugins/VirtualProduction/PerformanceCaptureWorkflow/Source/PerformanceCaptureWorkflowRuntime/Public/PCapPropComponent.h b/Engine/Plugins/VirtualProduction/PerformanceCaptureWorkflow/Source/PerformanceCaptureWorkflowRuntime/Public/PCapPropComponent.h
--- a/Engine/Plugins/VirtualProduction/PerformanceCaptureWorkflow/Source/PerformanceCaptureWorkflowRuntime/Public/PCapPropComponent.h
+++ b/Engine/Plugins/VirtualProduction/PerformanceCaptureWorkflow/Source/PerformanceCaptureWorkflowRuntime/Public/PCapPropComponent.h
@@ -142,7 +142,7 @@ public:
 	/**
 	 * Blueprint implementable event to enable calculation of a location offset before applying it to
 	 * the controlled component's transform.
-	 * @return World-space location offset to apply to the prop's transform
+	 * @return Parent-space location offset to apply to the prop's transform
 	 */
 	UFUNCTION(BlueprintNativeEvent, Category = "Performance Capture|Dynamic Constraint")
 	FVector CalculateDynamicOffset();
