[Performance Capture Workflow]

- Improved/clarified tool tip strings on function library functions and project settings.

diff --git a/Engine/Plugins/VirtualProduction/PerformanceCaptureWorkflow/Source/PerformanceCaptureWorkflow/Private/PCapSettings.h b/Engine/Plugins/VirtualProduction/PerformanceCaptureWorkflow/Source/PerformanceCaptureWorkflow/Private/PCapSettings.h
--- a/Engine/Plugins/VirtualProduction/PerformanceCaptureWorkflow/Source/PerformanceCaptureWorkflow/Private/PCapSettings.h
+++ b/Engine/Plugins/VirtualProduction/PerformanceCaptureWorkflow/Source/PerformanceCaptureWorkflow/Private/PCapSettings.h
@@ -48,11 +48,11 @@ public:
 	UPROPERTY(BlueprintReadOnly, EditDefaultsOnly, Config, NoClear, Category="Actors", DisplayName= "Default Performer Mesh")
 	TSoftObjectPtr<USkeletalMesh> DefaultPerformerSkelMesh;
 
-	/** Skeletal Mesh Actor classes that are explicitly allowed for use with Mocap Manager. */
+	/** Skeletal Mesh Actor classes that are explicitly to be used as Characters with Mocap Manager. These skeletal mesh actors must contain a retarget component*/
 	UPROPERTY(BlueprintReadOnly, EditDefaultsOnly, Config, NoClear, Category = "Actors")
 	TArray<TSoftClassPtr<ASkeletalMeshActor>> AllowedCaptureCharacterActorClasses;
 
-	/** Skeletal Mesh Actor classes that are explicitly disallowed for use with Mocap Manager. */
+	/** Skeletal Mesh Actor classes that are explicitly disallowed for use as Characters with Mocap Manager. */
 	UPROPERTY(BlueprintReadOnly, EditDefaultsOnly, Config, NoClear, Category = "Actors")
 	TArray<TSoftClassPtr<ASkeletalMeshActor>> DisallowedCaptureCharacterActorClasses;
 
diff --git a/Engine/Plugins/VirtualProduction/PerformanceCaptureWorkflow/Source/PerformanceCaptureWorkflowRuntime/Private/PCapWorkflowRuntimeFunctionLibrary.h b/Engine/Plugins/VirtualProduction/PerformanceCaptureWorkflow/Source/PerformanceCaptureWorkflowRuntime/Private/PCapWorkflowRuntimeFunctionLibrary.h
--- a/Engine/Plugins/VirtualProduction/PerformanceCaptureWorkflow/Source/PerformanceCaptureWorkflowRuntime/Private/PCapWorkflowRuntimeFunctionLibrary.h
+++ b/Engine/Plugins/VirtualProduction/PerformanceCaptureWorkflow/Source/PerformanceCaptureWorkflowRuntime/Private/PCapWorkflowRuntimeFunctionLibrary.h
@@ -7,21 +7,21 @@
 #include "PCapWorkflowRuntimeFunctionLibrary.generated.h"
 
 /**
- * Function for Performance Capture worfklow, enabling read-only access to Retargeter properties.
- * Intended for querying retargeters for calculation of dynamic prop constraints. 
+ * Function Library for Performance Capture Workflow, enabling read-only access to Retargeter properties.
+ * Intended for querying retargeters when calculating dynamic prop constraints. 
  */
 UCLASS()
 class PERFORMANCECAPTUREWORKFLOWRUNTIME_API UPCapWorkflowRuntimeFunctionLibrary : public UBlueprintFunctionLibrary
 {
 	GENERATED_BODY()
 	/**
-	 * Get the source IK rig from a IKRetargeter asset.
+	 * Get the source IK rig from an IKRetargeter asset.
 	 */
 	UFUNCTION(BlueprintCallable, Category="Performance Capture|Runtime")
 	static const UIKRigDefinition* GetSourceRig(const UIKRetargeter* InRetargetAsset);
 
 	/**
-	 * Get the target IK rig from a IKRetargeter asset.
+	 * Get the target IK rig from an IKRetargeter asset.
 	 */
 	UFUNCTION(BlueprintCallable, Category="Performance Capture|Runtime")
 	static const UIKRigDefinition* GetTargetRig(const UIKRetargeter* InRetargetAsset);
