[Performance Capture Workflow]
Enable users to choose what class of prop component to use. Default is the native version but users can override with their own derived class for uses such as building custom constraint behaviours.

[RN]

diff --git a/Engine/Plugins/VirtualProduction/PerformanceCaptureWorkflow/Source/PerformanceCaptureWorkflow/Private/PCapActorFactory.cpp b/Engine/Plugins/VirtualProduction/PerformanceCaptureWorkflow/Source/PerformanceCaptureWorkflow/Private/PCapActorFactory.cpp
--- a/Engine/Plugins/VirtualProduction/PerformanceCaptureWorkflow/Source/PerformanceCaptureWorkflow/Private/PCapActorFactory.cpp
+++ b/Engine/Plugins/VirtualProduction/PerformanceCaptureWorkflow/Source/PerformanceCaptureWorkflow/Private/PCapActorFactory.cpp
@@ -295,8 +295,9 @@ bool UPCapPropActorFactory::CanCreateActorFrom(const FAssetData& AssetData, FTex
 			const UStaticMesh* StaticMesh = PropDataAsset->PropStaticMesh.LoadSynchronous();
 			const USkeletalMesh* SkeletalMesh = PropDataAsset->PropSkeletalMesh.LoadSynchronous();;
 			const UClass* CustomProp = PropDataAsset->CustomPropClass.LoadSynchronous();
+			const UClass* PropComponent = PropDataAsset->PropComponentClass.LoadSynchronous();
 						
-			if(SkeletalMesh || StaticMesh || CustomProp)
+			if((SkeletalMesh || StaticMesh || CustomProp) && PropComponent )
 			{
 				return true;
 			}
@@ -310,7 +311,7 @@ void UPCapPropActorFactory::PostSpawnActor(UObject* Asset, AActor* NewActor)
 {
 	Super::PostSpawnActor(Asset, NewActor);
 	
-	if(NewActor)
+	if(NewActor && Asset)
 	{
 		if(UPCapPropDataAsset* PropDataAsset = Cast<UPCapPropDataAsset>(Asset))
 		{
@@ -365,7 +366,7 @@ void UPCapPropActorFactory::PostSpawnActor(UObject* Asset, AActor* NewActor)
 					FAddNewSubobjectParams NewObjectParams;
 
 					NewObjectParams.ParentHandle = Handle;
-					NewObjectParams.NewClass = UPCapPropComponent::StaticClass();
+					NewObjectParams.NewClass = PropDataAsset->PropComponentClass.Get();
 					NewObjectParams.bConformTransformToParent = true;
 
 					FText FailureReason;
@@ -394,5 +395,4 @@ void UPCapPropActorFactory::PostSpawnActor(UObject* Asset, AActor* NewActor)
 	}
 }
 
-#undef LOCTEXT_NAMESPACE
-
+#undef LOCTEXT_NAMESPACE
\ No newline at end of file
diff --git a/Engine/Plugins/VirtualProduction/PerformanceCaptureWorkflow/Source/PerformanceCaptureWorkflow/Private/PCapDatabase.cpp b/Engine/Plugins/VirtualProduction/PerformanceCaptureWorkflow/Source/PerformanceCaptureWorkflow/Private/PCapDatabase.cpp
--- a/Engine/Plugins/VirtualProduction/PerformanceCaptureWorkflow/Source/PerformanceCaptureWorkflow/Private/PCapDatabase.cpp
+++ b/Engine/Plugins/VirtualProduction/PerformanceCaptureWorkflow/Source/PerformanceCaptureWorkflow/Private/PCapDatabase.cpp
@@ -4,6 +4,7 @@
 #include "Algo/Transform.h"
 #include "CapturePerformer.h"
 #include "CaptureCharacter.h"
+#include "PCapPropComponent.h"
 #include "PCapSettings.h"
 #include "PerformanceCapture.h"
 #include "Misc/Guid.h"
@@ -93,6 +94,7 @@ TArray<UClass*> UPCapCharacterDataAsset::GetDisallowedCaptureCharacterActorClass
 
 UPCapPropDataAsset::UPCapPropDataAsset()
 	: UPCapDataAsset()
+	, PropComponentClass(UPCapPropComponent::StaticClass())
 {
 }
 
diff --git a/Engine/Plugins/VirtualProduction/PerformanceCaptureWorkflow/Source/PerformanceCaptureWorkflow/Private/PCapDatabase.h b/Engine/Plugins/VirtualProduction/PerformanceCaptureWorkflow/Source/PerformanceCaptureWorkflow/Private/PCapDatabase.h
--- a/Engine/Plugins/VirtualProduction/PerformanceCaptureWorkflow/Source/PerformanceCaptureWorkflow/Private/PCapDatabase.h
+++ b/Engine/Plugins/VirtualProduction/PerformanceCaptureWorkflow/Source/PerformanceCaptureWorkflow/Private/PCapDatabase.h
@@ -25,6 +25,7 @@ class UIKRetargeter;
 class ULiveLinkInstance;
 class UDataLayerAsset;
 class USkeletalMesh;
+class UPCapPropComponent;
 
 /// Database structs ///
 
@@ -425,6 +426,10 @@ public:
 	/** Static Mesh Offset Transform **/
 	UPROPERTY(BlueprintReadWrite, EditDefaultsOnly, Category="Performance Capture Prop")
 	FTransform PropOffsetTransform;
+
+	/** Prop Component - select a customized prop component if required, otherwise base C++ class will be used*/
+	UPROPERTY(BlueprintReadWrite, EditAnywhere, Category="Performance Capture Prop", meta=(DisplayName="Prop Component Class"), NoClear)
+	TSoftClassPtr<UPCapPropComponent> PropComponentClass;
 	
 	/** Prop Static - if this is left blank the mocap static mesh will be used*/
 	UPROPERTY(BlueprintReadWrite, EditDefaultsOnly, Category="Performance Capture Prop", meta=(DisplayName="Static Mesh", EditCondition="bCanSetStaticMesh", HideEditConditionToggle))
