MocapManager: QoL - Moving inline static to source for MessageLogName in PerformanceCaptureWorkflow module.

#rnx

diff --git a/Engine/Plugins/VirtualProduction/PerformanceCaptureWorkflow/Source/PerformanceCaptureWorkflow/Private/PerformanceCapture.cpp b/Engine/Plugins/VirtualProduction/PerformanceCaptureWorkflow/Source/PerformanceCaptureWorkflow/Private/PerformanceCapture.cpp
--- a/Engine/Plugins/VirtualProduction/PerformanceCaptureWorkflow/Source/PerformanceCaptureWorkflow/Private/PerformanceCapture.cpp
+++ b/Engine/Plugins/VirtualProduction/PerformanceCaptureWorkflow/Source/PerformanceCaptureWorkflow/Private/PerformanceCapture.cpp
@@ -22,6 +22,8 @@ static const FName PerformanceCapturePanelTabName("PerformanceCaptureTab");
 
 DEFINE_LOG_CATEGORY(LogPCap);
 
+const FLazyName FPerformanceCaptureModule::MessageLogName = TEXT("PerformanceCapture");
+
 void FPerformanceCaptureModule::StartupModule()
 {
 	if(ISettingsModule* SettingsModule = FModuleManager::GetModulePtr<ISettingsModule>("Settings"))
diff --git a/Engine/Plugins/VirtualProduction/PerformanceCaptureWorkflow/Source/PerformanceCaptureWorkflow/Public/PerformanceCapture.h b/Engine/Plugins/VirtualProduction/PerformanceCaptureWorkflow/Source/PerformanceCaptureWorkflow/Public/PerformanceCapture.h
--- a/Engine/Plugins/VirtualProduction/PerformanceCaptureWorkflow/Source/PerformanceCaptureWorkflow/Public/PerformanceCapture.h
+++ b/Engine/Plugins/VirtualProduction/PerformanceCaptureWorkflow/Source/PerformanceCaptureWorkflow/Public/PerformanceCapture.h
@@ -6,6 +6,8 @@
 #include "IPlacementModeModule.h"
 #include "Modules/ModuleManager.h"
 
+#define UE_API PERFORMANCECAPTUREWORKFLOW_API
+
 DECLARE_LOG_CATEGORY_EXTERN(LogPCap, Log, All);
 class FPlacementModeID;
 
@@ -21,7 +23,7 @@ public:
 	void PluginButtonClicked();
 
 	/** Message Log Name for the PerformanceCapture module. */
-	inline static const FLazyName MessageLogName = TEXT("PerformanceCapture");
+	UE_API static const FLazyName MessageLogName; // "PerformanceCapture"
 
 private:
 
@@ -31,3 +33,5 @@ private:
 	
 	TSharedPtr<class FUICommandList> PluginCommands;
 };
+
+#undef UE_API
