[Performance Capture Workflow]
- Enable bIsControlledBySequencer to be keyable in sequencer with Interp UPROPERTY
- Add OnSubjectRemoved multicast delegate
- Add OnSubjectEnabled multicast delegate
- Improved docstrings

#rnx

diff --git a/Engine/Plugins/VirtualProduction/PerformanceCaptureWorkflow/Source/PerformanceCaptureWorkflow/Private/PCapSubsystem.cpp b/Engine/Plugins/VirtualProduction/PerformanceCaptureWorkflow/Source/PerformanceCaptureWorkflow/Private/PCapSubsystem.cpp
--- a/Engine/Plugins/VirtualProduction/PerformanceCaptureWorkflow/Source/PerformanceCaptureWorkflow/Private/PCapSubsystem.cpp
+++ b/Engine/Plugins/VirtualProduction/PerformanceCaptureWorkflow/Source/PerformanceCaptureWorkflow/Private/PCapSubsystem.cpp
@@ -68,6 +68,8 @@ void UPerformanceCaptureSubsystem::OnEngineInitComplete()
 	{
 		ILiveLinkClient& LiveLinkClient = IModularFeatures::Get().GetModularFeature<ILiveLinkClient>(ILiveLinkClient::ModularFeatureName);
 		LiveLinkClient.OnLiveLinkSubjectStateChanged().AddUObject(this, &UPerformanceCaptureSubsystem::OnLiveLinkSubjectUpdated);
+		LiveLinkClient.OnLiveLinkSubjectRemoved().AddUObject(this, &UPerformanceCaptureSubsystem::OnLiveLinkSubjectRemoved);
+		LiveLinkClient.OnLiveLinkSubjectEnabledChanged().AddUObject(this, &UPerformanceCaptureSubsystem::OnLiveLinkSubjectEnableChanged);
 	}
 }
 
@@ -150,7 +152,13 @@ void UPerformanceCaptureSubsystem::OnLiveLinkSubjectUpdated(FLiveLinkSubjectKey
 	OnPCapLiveLinkSubjectUpdate.Broadcast(Subject, State);
 }
 
+void UPerformanceCaptureSubsystem::OnLiveLinkSubjectRemoved(FLiveLinkSubjectKey Subject) const
+{
+	OnPCapLiveLinkSubjectRemoved.Broadcast(Subject);
+}
 
-
-
+void UPerformanceCaptureSubsystem::OnLiveLinkSubjectEnableChanged(FLiveLinkSubjectKey Subject, bool NewEnabled) const
+{
+	OnPCapLiveLinkSubjectEnableChanged.Broadcast(Subject, NewEnabled);
+}
 
diff --git a/Engine/Plugins/VirtualProduction/PerformanceCaptureWorkflow/Source/PerformanceCaptureWorkflow/Private/PCapSubsystem.h b/Engine/Plugins/VirtualProduction/PerformanceCaptureWorkflow/Source/PerformanceCaptureWorkflow/Private/PCapSubsystem.h
--- a/Engine/Plugins/VirtualProduction/PerformanceCaptureWorkflow/Source/PerformanceCaptureWorkflow/Private/PCapSubsystem.h
+++ b/Engine/Plugins/VirtualProduction/PerformanceCaptureWorkflow/Source/PerformanceCaptureWorkflow/Private/PCapSubsystem.h
@@ -55,51 +55,94 @@ public:
 
 #if WITH_EDITORONLY_DATA
 	
+	/**
+	* Multicast delegate that will execute when an asset is removed from the asset registry. Returns the FAsssetData of the removed asset.
+	*/
 	DECLARE_DYNAMIC_MULTICAST_DELEGATE_OneParam(FOnPCapAssetRemoved, FAssetData, DeletedAsset);
 	/** Editor Only*/
 	UPROPERTY(BlueprintAssignable, Category="Performance Capture|Assets")
 	FOnPCapAssetRemoved OnPCapAssetRemoved;
 
+	/**
+	* Multicast delegate that will execute when an asset is renamed to the asset registry. Returns the FAssetData of the renamed asset.
+	*/
 	DECLARE_DYNAMIC_MULTICAST_DELEGATE_TwoParams(FOnPCapAssetRenamed, FAssetData, RenamedAsset, FString, OldName);
 	/** Editor Only*/
 	UPROPERTY(BlueprintAssignable, Category="Performance Capture|Assets")
 	FOnPCapAssetRenamed OnPCapAssetRenamed;
 
+	/**
+	* Multicast delegate that will execute when a new asset is added to the asset registry. Returns the FAssetData of new asset.
+	*/
 	DECLARE_DYNAMIC_MULTICAST_DELEGATE_OneParam(FOnPCapAssetAdded, FAssetData, NewAsset);
 	/** Editor Only*/
 	UPROPERTY(BlueprintAssignable, Category="Performance Capture|Assets")
 	FOnPCapAssetAdded OnPCapAssetAdded;
 
+	/**
+	* Multicast delegate that will execute when an actor is modified in the Level Editor. Returns the actor modified.
+	*/
 	DECLARE_DYNAMIC_MULTICAST_DELEGATE_OneParam(FOnPCapActorModified, AActor*, Actor);
 	/** Editor Only*/
 	UPROPERTY(BlueprintAssignable, Category="Performance Capture|Workflow")
 	FOnPCapActorModified OnPCapActorModified;
-
+	
+	/**
+	* Multicast delegate that will execute when the user calls Redo. Returns success boolean.
+	*/
 	DECLARE_DYNAMIC_MULTICAST_DELEGATE_OneParam(FOnPCapEditorUndo, bool, bSuccess);
 	/** Editor Only*/
 	UPROPERTY(BlueprintAssignable, Category="Performance Capture|Workflow")
 	FOnPCapEditorUndo OnPCapEditorUndo;
 
+	/**
+	* Multicast delegate that will execute when the user calls Redo. Returns success boolean.
+	*/
 	DECLARE_DYNAMIC_MULTICAST_DELEGATE_OneParam(FOnPCapEditorRedo, bool, bSuccess);
 	/** Editor Only*/
 	UPROPERTY(BlueprintAssignable, Category="Performance Capture|Workflow")
 	FOnPCapEditorRedo OnPCapEditorRedo;
 
+	/**
+	* Multicast delegate that will execute when an asset editor is opened. Returns the asset class that is opened. 
+	*/
 	DECLARE_DYNAMIC_MULTICAST_DELEGATE_OneParam(FOnPCapAssetEditorOpen, UObject*, Asset);
 	/** Editor Only*/
 	UPROPERTY(BlueprintAssignable, Category="Performance Capture|Assets")
 	FOnPCapAssetEditorOpen OnPCapAssetEditorOpen;
 
+	/**
+	* Multicast delegate that will execute when an asset editor is closed. Returns the asset class that is closed. 
+	*/
 	DECLARE_DYNAMIC_MULTICAST_DELEGATE_OneParam(FOnPCapAssetEditorClose, UObject*, Asset);
 	/** Editor Only*/
 	UPROPERTY(BlueprintAssignable, Category="Performance Capture|Assets")
 	FOnPCapAssetEditorClose OnPCapAssetEditorClose;
 
+	/**
+	* Multicast delegate that will execute when a Live Link subject's state is updated.
+	*/
 	DECLARE_DYNAMIC_MULTICAST_DELEGATE_TwoParams(FOnPCapLiveLinkSubjectUpdate, FLiveLinkSubjectKey, Subject, ELiveLinkSubjectState, State);
 	/** Editor Only*/
 	UPROPERTY(BlueprintAssignable, Category="Performance Capture|LiveLink")
 	FOnPCapLiveLinkSubjectUpdate OnPCapLiveLinkSubjectUpdate;
 
+	/**
+	* Multicast delegate that will execute when a Live Link subject is removed.
+	*/
+	DECLARE_DYNAMIC_MULTICAST_DELEGATE_OneParam(FOnPCapLiveLinkSubjectRemoved, FLiveLinkSubjectKey, Subject);
+	/** Editor Only*/
+	UPROPERTY(BlueprintAssignable, Category="Performance Capture|LiveLink")
+	FOnPCapLiveLinkSubjectRemoved OnPCapLiveLinkSubjectRemoved;
+
+	/**
+	* Multicast delegate that will execute when a Live Link subject's bIsEnabled member is modified in the Live Link panel.
+	*/
+	DECLARE_DYNAMIC_MULTICAST_DELEGATE_TwoParams(FOnPCapLiveLinkSubjectEnableChanged, FLiveLinkSubjectKey, Subject, bool, NewEnabled);
+	/** Editor Only*/
+	UPROPERTY(BlueprintAssignable, Category="Performance Capture|LiveLink")
+	FOnPCapLiveLinkSubjectEnableChanged OnPCapLiveLinkSubjectEnableChanged;
+
 
 #endif
 
@@ -141,6 +184,10 @@ private:
 	void OnAssetClosed(UObject* Object, IAssetEditorInstance* Instance) const;
 
 	void OnLiveLinkSubjectUpdated(FLiveLinkSubjectKey Subject, ELiveLinkSubjectState State) const;
+
+	void OnLiveLinkSubjectRemoved(FLiveLinkSubjectKey Subject) const;
+
+	void OnLiveLinkSubjectEnableChanged(FLiveLinkSubjectKey Subject, bool NewEnabled) const;
 };
 
 //Todo: add post edit change to re-int the viewmodel and database helper if the user changes these once the editor has started. 
diff --git a/Engine/Plugins/VirtualProduction/PerformanceCaptureWorkflow/Source/PerformanceCaptureWorkflowRuntime/Public/PCapPropComponent.h b/Engine/Plugins/VirtualProduction/PerformanceCaptureWorkflow/Source/PerformanceCaptureWorkflowRuntime/Public/PCapPropComponent.h
--- a/Engine/Plugins/VirtualProduction/PerformanceCaptureWorkflow/Source/PerformanceCaptureWorkflowRuntime/Public/PCapPropComponent.h
+++ b/Engine/Plugins/VirtualProduction/PerformanceCaptureWorkflow/Source/PerformanceCaptureWorkflowRuntime/Public/PCapPropComponent.h
@@ -39,11 +39,11 @@ public:
 	/**
 	 * Overrides all LiveLink data and the offset transform. Must be manually set on possessable bindings.
 	 */
-	UPROPERTY(BlueprintReadOnly, EditInstanceOnly, Category = "Performance Capture")
+	UPROPERTY(BlueprintReadOnly, EditInstanceOnly, Interp, Category = "Performance Capture")
 	bool bIsControlledBySequencer = false;
 
 	/**
-	 *The package name of the spawning data asset, if there is one. 
+	 * The package name of the spawning data asset, if there is one. 
 	 */
 	UPROPERTY(BlueprintReadOnly, VisibleAnywhere, Category = "Performance Capture")
 	FName SpawningDataAsset;
