[Performance Capture Workflow]
- Bughawk docstring fixes.

#rnx

diff --git a/Engine/Plugins/VirtualProduction/PerformanceCaptureWorkflow/Source/PerformanceCaptureWorkflow/Private/PCapSubsystem.h b/Engine/Plugins/VirtualProduction/PerformanceCaptureWorkflow/Source/PerformanceCaptureWorkflow/Private/PCapSubsystem.h
--- a/Engine/Plugins/VirtualProduction/PerformanceCaptureWorkflow/Source/PerformanceCaptureWorkflow/Private/PCapSubsystem.h
+++ b/Engine/Plugins/VirtualProduction/PerformanceCaptureWorkflow/Source/PerformanceCaptureWorkflow/Private/PCapSubsystem.h
@@ -56,7 +56,7 @@ public:
 #if WITH_EDITORONLY_DATA
 	
 	/**
-	* Multicast delegate that will execute when an asset is removed from the asset registry. Returns the FAsssetData of the removed asset.
+	* Multicast delegate that will execute when an asset is removed from the asset registry. Returns the FAssetData of the removed asset.
 	*/
 	DECLARE_DYNAMIC_MULTICAST_DELEGATE_OneParam(FOnPCapAssetRemoved, FAssetData, DeletedAsset);
 	/** Editor Only*/
@@ -88,7 +88,7 @@ public:
 	FOnPCapActorModified OnPCapActorModified;
 	
 	/**
-	* Multicast delegate that will execute when the user calls Redo. Returns success boolean.
+	* Multicast delegate that will execute when the user calls Undo. Returns success boolean.
 	*/
 	DECLARE_DYNAMIC_MULTICAST_DELEGATE_OneParam(FOnPCapEditorUndo, bool, bSuccess);
 	/** Editor Only*/
@@ -104,7 +104,7 @@ public:
 	FOnPCapEditorRedo OnPCapEditorRedo;
 
 	/**
-	* Multicast delegate that will execute when an asset editor is opened. Returns the asset class that is opened. 
+	* Multicast delegate that will execute when an asset editor is opened. Returns the asset object that is opened. 
 	*/
 	DECLARE_DYNAMIC_MULTICAST_DELEGATE_OneParam(FOnPCapAssetEditorOpen, UObject*, Asset);
 	/** Editor Only*/
@@ -112,7 +112,7 @@ public:
 	FOnPCapAssetEditorOpen OnPCapAssetEditorOpen;
 
 	/**
-	* Multicast delegate that will execute when an asset editor is closed. Returns the asset class that is closed. 
+	* Multicast delegate that will execute when an asset editor is closed. Returns the asset object that is closed. 
 	*/
 	DECLARE_DYNAMIC_MULTICAST_DELEGATE_OneParam(FOnPCapAssetEditorClose, UObject*, Asset);
 	/** Editor Only*/
