[Performance Capture Workflow]
- Marked Root folder on Session Template as BlueprintReadWrite

#rnx

diff --git a/Engine/Plugins/VirtualProduction/PerformanceCaptureWorkflow/Source/PerformanceCaptureWorkflow/Private/PCapSessionTemplate.h b/Engine/Plugins/VirtualProduction/PerformanceCaptureWorkflow/Source/PerformanceCaptureWorkflow/Private/PCapSessionTemplate.h
--- a/Engine/Plugins/VirtualProduction/PerformanceCaptureWorkflow/Source/PerformanceCaptureWorkflow/Private/PCapSessionTemplate.h
+++ b/Engine/Plugins/VirtualProduction/PerformanceCaptureWorkflow/Source/PerformanceCaptureWorkflow/Private/PCapSessionTemplate.h
@@ -59,7 +59,7 @@ public:
 	* Core session fields
 	*/
 
-	UPROPERTY(EditAnywhere, BlueprintReadOnly, Category="Template", meta = (DisplayName = "Root Folder", ContentDir, EditCondition="bIsEditable"))
+	UPROPERTY(EditAnywhere, BlueprintReadWrite, Category="Template", meta = (DisplayName = "Root Folder", ContentDir, EditCondition="bIsEditable"))
 	FDirectoryPath TemplateRootFolder;
 
 	/*
