[Performance Capture Workflow]
Add BP setter on the PCap settings object to enable setting the default session template from the mocap manager panel.
Setter emits a post edit change so BP can bind and react to changes.
Added utility object to provide user with a content folder picker to choose a custom PCap data folder. Done because directory path needs "ContentDir" UProperty meta.

#rnx

diff --git a/Engine/Plugins/VirtualProduction/PerformanceCaptureWorkflow/Config/BasePerformanceCaptureWorkflow.ini b/Engine/Plugins/VirtualProduction/PerformanceCaptureWorkflow/Config/BasePerformanceCaptureWorkflow.ini
--- a/Engine/Plugins/VirtualProduction/PerformanceCaptureWorkflow/Config/BasePerformanceCaptureWorkflow.ini
+++ b/Engine/Plugins/VirtualProduction/PerformanceCaptureWorkflow/Config/BasePerformanceCaptureWorkflow.ini
@@ -5,7 +5,6 @@ bCanSaveAllSections=True
 StageRoot=/PerformanceCaptureWorkflow/Stage/BP_StageRoot.BP_StageRoot_C
 ViewModelClass=/PerformanceCaptureWorkflow/Core/VM_PCapWorkflow.VM_PCapWorkflow_C
 MocapManagerUI=/PerformanceCaptureWorkflow/Core/EW_MocapManager.EW_MocapManager
-DefaultSessionTemplate=/PerformanceCaptureWorkflow/Core/DefaultSessionTemplate.DefaultSessionTemplate
 DatabaseHelperClass=/PerformanceCaptureWorkflow/Database/EU_InternalDatabaseHelper.EU_InternalDatabaseHelper_C
 AllowedCaptureCharacterActorClasses=/Script/PerformanceCaptureCore.CaptureCharacter
 DisallowedCaptureCharacterActorClasses=/Script/PerformanceCaptureCore.CapturePerformer
diff --git a/Engine/Plugins/VirtualProduction/PerformanceCaptureWorkflow/Source/PerformanceCaptureWorkflow/Private/PCapSettings.cpp b/Engine/Plugins/VirtualProduction/PerformanceCaptureWorkflow/Source/PerformanceCaptureWorkflow/Private/PCapSettings.cpp
--- a/Engine/Plugins/VirtualProduction/PerformanceCaptureWorkflow/Source/PerformanceCaptureWorkflow/Private/PCapSettings.cpp
+++ b/Engine/Plugins/VirtualProduction/PerformanceCaptureWorkflow/Source/PerformanceCaptureWorkflow/Private/PCapSettings.cpp
@@ -31,20 +31,27 @@ void UPerformanceCaptureSettings::SetSessionTable(TSoftObjectPtr<UPCapDataTable>
 	SessionTable = NewDataTable;
 
 	FProperty* SessionTableProperty = FindFieldChecked<FProperty>(UPerformanceCaptureSettings::StaticClass(),GET_MEMBER_NAME_CHECKED(UPerformanceCaptureSettings, SessionTable));
-	FPropertyChangedEvent  PCapSessionTableChangedEvent = FPropertyChangedEvent(SessionTableProperty);
+	FPropertyChangedEvent  PCapSessionTableChangedEvent = FPropertyChangedEvent(SessionTableProperty, EPropertyChangeType::ValueSet);
 	PostEditChangeProperty(PCapSessionTableChangedEvent);
 	TryUpdateDefaultConfigFile();
-	SaveConfig();
 }
 
 void UPerformanceCaptureSettings::SetProductionTable(TSoftObjectPtr<UPCapDataTable> NewDataTable)
 {
 	ProductionTable = NewDataTable;
 	FProperty* ProductionTableProperty = FindFieldChecked<FProperty>(UPerformanceCaptureSettings::StaticClass(),GET_MEMBER_NAME_CHECKED(UPerformanceCaptureSettings, ProductionTable));
-	FPropertyChangedEvent  PCapProductionTableChangedEvent = FPropertyChangedEvent(ProductionTableProperty);
+	FPropertyChangedEvent  PCapProductionTableChangedEvent = FPropertyChangedEvent(ProductionTableProperty, EPropertyChangeType::ValueSet);
 	PostEditChangeProperty(PCapProductionTableChangedEvent);
 	TryUpdateDefaultConfigFile();
-	SaveConfig();
+}
+
+void UPerformanceCaptureSettings::SetDefaultSessionTemplate(TSoftObjectPtr<UPCapSessionTemplate> NewSessionTemplate)
+{
+	DefaultSessionTemplate = NewSessionTemplate;
+	FProperty* DefaultSessionTemplateProperty = FindFieldChecked<FProperty>(UPerformanceCaptureSettings::StaticClass(), GET_MEMBER_NAME_CHECKED(UPerformanceCaptureSettings, DefaultSessionTemplate));
+	FPropertyChangedEvent  PCapDefaultSessionTemplateChangedEvent = FPropertyChangedEvent(DefaultSessionTemplateProperty, EPropertyChangeType::ValueSet);
+	PostEditChangeProperty(PCapDefaultSessionTemplateChangedEvent);
+	TryUpdateDefaultConfigFile();
 }
 
 
diff --git a/Engine/Plugins/VirtualProduction/PerformanceCaptureWorkflow/Source/PerformanceCaptureWorkflow/Private/PCapSettings.h b/Engine/Plugins/VirtualProduction/PerformanceCaptureWorkflow/Source/PerformanceCaptureWorkflow/Private/PCapSettings.h
--- a/Engine/Plugins/VirtualProduction/PerformanceCaptureWorkflow/Source/PerformanceCaptureWorkflow/Private/PCapSettings.h
+++ b/Engine/Plugins/VirtualProduction/PerformanceCaptureWorkflow/Source/PerformanceCaptureWorkflow/Private/PCapSettings.h
@@ -92,7 +92,14 @@ public:
 	 * @param NewDataTable Datatable ref. Must be of the FPCapProduction struct type.
 	 */
 	UFUNCTION(BlueprintCallable, Category="Performance Capture|Settings")
-	void SetProductionTable(TSoftObjectPtr<UPCapDataTable> NewDataTable); 
+	void SetProductionTable(TSoftObjectPtr<UPCapDataTable> NewDataTable);
+
+	/**
+	 * Set the default session template in Performance Capture Settings.
+	 * @param NewSessionTemplate Session Template data asset ref.
+	 */
+	UFUNCTION(BlueprintCallable, Category="Performance Capture|Settings")
+	void SetDefaultSessionTemplate(TSoftObjectPtr<UPCapSessionTemplate> NewSessionTemplate); 
 
 #if WITH_EDITOR
 	
@@ -107,3 +114,16 @@ public:
 	virtual void PostEditChangeProperty(FPropertyChangedEvent& PropertyChangedEvent) override;
 #endif
 };
+
+/**
+ * Utility object for choosing a content browser path in a message dialog.
+ */
+UCLASS(Blueprintable)
+class  UPCapDialogObject : public UObject
+{
+	GENERATED_BODY()
+
+public:
+	UPROPERTY(EditAnywhere, BlueprintReadWrite, Category=nocategory, meta=(ContentDir, DisplayName="PCap Path", ToolTip="Choose a folder for your Performance Capture data"))
+	FDirectoryPath Path;
+};
