[Performance Capture Workflow]
- Changed access pattern on RootFolder of template to match the other BP accessible members, using BlueprintSetter and Getter functions

#rnx

#virtualized

diff --git a/Engine/Plugins/VirtualProduction/PerformanceCaptureWorkflow/Source/PerformanceCaptureWorkflow/Private/PCapSessionTemplate.h b/Engine/Plugins/VirtualProduction/PerformanceCaptureWorkflow/Source/PerformanceCaptureWorkflow/Private/PCapSessionTemplate.h
--- a/Engine/Plugins/VirtualProduction/PerformanceCaptureWorkflow/Source/PerformanceCaptureWorkflow/Private/PCapSessionTemplate.h
+++ b/Engine/Plugins/VirtualProduction/PerformanceCaptureWorkflow/Source/PerformanceCaptureWorkflow/Private/PCapSessionTemplate.h
@@ -59,7 +59,7 @@ public:
 	* Core session fields
 	*/
 
-	UPROPERTY(EditAnywhere, BlueprintReadWrite, Category="Template", meta = (DisplayName = "Root Folder", ContentDir, EditCondition="bIsEditable"))
+	UPROPERTY(EditAnywhere, BlueprintSetter = SetRootFolder, BlueprintGetter = GetRootFolder, Category="Template", meta = (DisplayName = "Root Folder", ContentDir, EditCondition="bIsEditable"))
 	FDirectoryPath TemplateRootFolder;
 
 	/*
@@ -80,6 +80,26 @@ public:
 	UPROPERTY(EditAnywhere, BlueprintReadOnly, Category="Template", meta = (DisplayName = "Session Token", ContentDir, EditCondition="bIsEditable", tooltip="The final name of our session, including any tokens you include in it's construction. You can use tokens to build this string, including the Session Name above"))
 	FPCapTokenisedString SessionToken;
 
+	/**
+	 * Set the RootFolder.
+	 */
+	UFUNCTION(BlueprintSetter)
+	void SetRootFolder(FDirectoryPath NewRootPath)
+	{
+		TemplateRootFolder = NewRootPath;
+		UpdateAllFields();
+	}
+
+	/**
+	 * Get the RootFolder.
+	 * @return TemplateRootFolder path.
+	 */
+	UFUNCTION(BlueprintGetter)
+	FDirectoryPath GetRootFolder()
+	{
+		return TemplateRootFolder;
+	}
+	
 	/**
 	 * Get the Session Name.
 	 * @return SessionName string.
