[Performance Capture Workflow]
- Modified bone visualizer to work with any actor containing at least one skinnned mesh component. Always defaults to the first found skinned mesh.

diff --git a/Engine/Plugins/VirtualProduction/PerformanceCaptureWorkflow/Config/BasePerformanceCaptureWorkflow.ini b/Engine/Plugins/VirtualProduction/PerformanceCaptureWorkflow/Config/BasePerformanceCaptureWorkflow.ini
--- a/Engine/Plugins/VirtualProduction/PerformanceCaptureWorkflow/Config/BasePerformanceCaptureWorkflow.ini
+++ b/Engine/Plugins/VirtualProduction/PerformanceCaptureWorkflow/Config/BasePerformanceCaptureWorkflow.ini
@@ -44,4 +44,5 @@ DisallowedCaptureCharacterActorClasses=/Script/PerformanceCaptureCore.CapturePer
 +ClassRedirects=(OldName="/Script/PerformanceCaptureWorkflow.PerformanceCaptureStageRoot",NewName="/Script/PerformanceCaptureWorkflowRuntime.PerformanceCaptureStageRoot")
 +ClassRedirects=(OldName="/Script/PerformanceCaptureWorkflow.PCapPropComponent",NewName="/Script/PerformanceCaptureWorkflowRuntime.PCapPropComponent")
 +ClassRedirects=(OldName="/Script/PerformanceCaptureWorkflow.PCapPropLiveLinkAnimInstance",NewName="/Script/PerformanceCaptureWorkflowRuntime.PCapPropLiveLinkAnimInstance")
-+ClassRedirects=(OldName="/Script/PerformanceCaptureWorkflow.AnimNode_LiveLinkProp",NewName="/Script/PerformanceCaptureWorkflowRuntime.AnimNode_LiveLinkProp")
\ No newline at end of file
++ClassRedirects=(OldName="/Script/PerformanceCaptureWorkflow.AnimNode_LiveLinkProp",NewName="/Script/PerformanceCaptureWorkflowRuntime.AnimNode_LiveLinkProp")
++PropertyRedirects=(OldName="/Script/PerformanceCaptureWorkflow.PCapBoneVisualiser.SkelmeshComponent",NewName="/Script/PerformanceCaptureWorkflow.PCapBoneVisualiser.SkinnedMeshComponent")
\ No newline at end of file
diff --git a/Engine/Plugins/VirtualProduction/PerformanceCaptureWorkflow/Source/PerformanceCaptureWorkflow/Private/Visualizers/PCapBoneVisualizer.cpp b/Engine/Plugins/VirtualProduction/PerformanceCaptureWorkflow/Source/PerformanceCaptureWorkflow/Private/Visualizers/PCapBoneVisualizer.cpp
--- a/Engine/Plugins/VirtualProduction/PerformanceCaptureWorkflow/Source/PerformanceCaptureWorkflow/Private/Visualizers/PCapBoneVisualizer.cpp
+++ b/Engine/Plugins/VirtualProduction/PerformanceCaptureWorkflow/Source/PerformanceCaptureWorkflow/Private/Visualizers/PCapBoneVisualizer.cpp
@@ -25,18 +25,18 @@ void UPCapBoneVisualiser::OnRegister()
 {
 	Super::OnRegister();
 
-	ASkeletalMeshActor* OwnerActor = Cast<ASkeletalMeshActor>(GetOwner());
+	AActor* OwnerActor = GetOwner();
 	
 	if(OwnerActor)
 	{
-		SkelmeshComponent = OwnerActor->GetSkeletalMeshComponent();
+		SkinnedMeshComponent = OwnerActor->GetComponentByClass<USkinnedMeshComponent>();
 	}
 
 	ClearInstances();
 
-	if(SkelmeshComponent)
+	if(SkinnedMeshComponent)
 	{
-		TArray<FTransform> BoneTransforms = GetBoneTransforms(SkelmeshComponent);
+		TArray<FTransform> BoneTransforms = GetBoneTransforms(SkinnedMeshComponent);
 		AddInstances(BoneTransforms, false, true, false);
 
 		if(!DynamicMaterial)
@@ -64,18 +64,18 @@ void UPCapBoneVisualiser::TickComponent(float DeltaTime, ELevelTick TickType, FA
 {
 	Super::TickComponent(DeltaTime, TickType, ThisTickFunction);
 	
-	if(SkelmeshComponent)
+	if(SkinnedMeshComponent)
 	{
 		switch (VisualizationType)
 		{
 		case EBoneVisualType::Bone:
 			{
-				BatchUpdateInstancesTransforms(0, GetBoneTransforms(SkelmeshComponent), true, true, false);
+				BatchUpdateInstancesTransforms(0, GetBoneTransforms(SkinnedMeshComponent), true, true, false);
 			}
 			break;
 		case EBoneVisualType::Joint:
 			{
-				BatchUpdateInstancesTransforms(0, GetJointTransforms(SkelmeshComponent), true, true, false);
+				BatchUpdateInstancesTransforms(0, GetJointTransforms(SkinnedMeshComponent), true, true, false);
 			}
 
 			break;
@@ -85,41 +85,43 @@ void UPCapBoneVisualiser::TickComponent(float DeltaTime, ELevelTick TickType, FA
 	}
 }
 
-TArray<FTransform> UPCapBoneVisualiser::GetJointTransforms(USkeletalMeshComponent* InSkelMeshComponent) const
+TArray<FTransform> UPCapBoneVisualiser::GetJointTransforms(USkinnedMeshComponent* InSkinnedMeshComponent) const
 {
 	TArray<FTransform> FoundTransforms;
-
-	int32 i = InSkelMeshComponent->GetNumBones();
-	FTransform BoneTransform;
-
-	for (i = 0; i < InSkelMeshComponent->GetNumBones(); i++ )
+	
+	if(InSkinnedMeshComponent)
 	{
-		BoneTransform = InSkelMeshComponent->GetBoneTransform(InSkelMeshComponent->GetBoneName(i));
-		FoundTransforms.Add(BoneTransform);
-	}
+		FTransform BoneTransform;
 
+		for (int32 i = 0; i < InSkinnedMeshComponent->GetNumBones(); i++ )
+		{
+			BoneTransform = InSkinnedMeshComponent->GetBoneTransform(InSkinnedMeshComponent->GetBoneName(i));
+			FoundTransforms.Add(BoneTransform);
+		}
+	}
 	return FoundTransforms;
 }
 
-TArray<FTransform> UPCapBoneVisualiser::GetBoneTransforms(USkeletalMeshComponent* InSkelMeshComponent) const
+TArray<FTransform> UPCapBoneVisualiser::GetBoneTransforms(USkinnedMeshComponent* InSkinnedMeshComponent) const
 {
 	TArray<FTransform> FoundTransforms;
-
-	int32 i = InSkelMeshComponent->GetNumBones();
-	FTransform BoneTransform;
-	FTransform ParentBoneTransform;
-	double BoneLength;
-	FVector Aim = FVector( 0, 0, 1);
-
-	for (i = 0; i < InSkelMeshComponent->GetNumBones(); i++ )
+	if(InSkinnedMeshComponent)
 	{
-		BoneTransform = InSkelMeshComponent->GetBoneTransform(InSkelMeshComponent->GetBoneName(i));
-		ParentBoneTransform = InSkelMeshComponent->GetBoneTransform(InSkelMeshComponent->GetParentBone(InSkelMeshComponent->GetBoneName(i)));
-		BoneLength = FVector::Distance(BoneTransform.GetLocation(), ParentBoneTransform.GetLocation());
-		BoneTransform.SetScale3D(FVector(1, 1, BoneLength));
-		FQuat DiffRotation = AnimationCore::SolveAim(BoneTransform, ParentBoneTransform.GetLocation(), Aim.GetSafeNormal(), false, FVector(1, 1, 1), float(0));
-		BoneTransform.SetRotation(DiffRotation);
-		FoundTransforms.Add(BoneTransform);
+		FTransform BoneTransform;
+		FTransform ParentBoneTransform;
+		double BoneLength;
+		FVector Aim = FVector( 0, 0, 1);
+
+		for (int32 i = 0; i < InSkinnedMeshComponent->GetNumBones(); i++ )
+		{
+			BoneTransform = InSkinnedMeshComponent->GetBoneTransform(InSkinnedMeshComponent->GetBoneName(i));
+			ParentBoneTransform = InSkinnedMeshComponent->GetBoneTransform(InSkinnedMeshComponent->GetParentBone(InSkinnedMeshComponent->GetBoneName(i)));
+			BoneLength = FVector::Distance(BoneTransform.GetLocation(), ParentBoneTransform.GetLocation());
+			BoneTransform.SetScale3D(FVector(1, 1, BoneLength));
+			FQuat DiffRotation = AnimationCore::SolveAim(BoneTransform, ParentBoneTransform.GetLocation(), Aim.GetSafeNormal(), false, FVector(1, 1, 1), float(0));
+			BoneTransform.SetRotation(DiffRotation);
+			FoundTransforms.Add(BoneTransform);
+		}
 	}
 	return FoundTransforms;
 }
@@ -146,3 +148,4 @@ void UPCapBoneVisualiser::PostEditChangeProperty(FPropertyChangedEvent& Property
 	}
 }
 
+
diff --git a/Engine/Plugins/VirtualProduction/PerformanceCaptureWorkflow/Source/PerformanceCaptureWorkflow/Private/Visualizers/PCapBoneVisualizer.h b/Engine/Plugins/VirtualProduction/PerformanceCaptureWorkflow/Source/PerformanceCaptureWorkflow/Private/Visualizers/PCapBoneVisualizer.h
--- a/Engine/Plugins/VirtualProduction/PerformanceCaptureWorkflow/Source/PerformanceCaptureWorkflow/Private/Visualizers/PCapBoneVisualizer.h
+++ b/Engine/Plugins/VirtualProduction/PerformanceCaptureWorkflow/Source/PerformanceCaptureWorkflow/Private/Visualizers/PCapBoneVisualizer.h
@@ -6,7 +6,7 @@
 #include "Components/InstancedStaticMeshComponent.h"
 #include "PCapBoneVisualizer.generated.h"
 
-class USkeletalMeshComponent;
+class USkinnedMeshComponent;
 class UMaterialInstanceDynamic;
 
 UENUM(BlueprintType)
@@ -16,8 +16,8 @@ enum class EBoneVisualType : uint8
 	Bone	UMETA(DisplayName = "Bone"),
 };
 
-/** Instanced Static Mesh Component for drawing bones and joints on Skeletal Mesh. Only usable with SkeletalMesh Actors*/
-UCLASS(ClassGroup=(PerformanceCapture), meta=(BlueprintSpawnableComponent),Transient, DisplayName = "PerformanceCaptureVisualizer")
+/** Instanced Static Mesh Component for drawing bones and joints on a Skinned Mesh component. Only usable with Actors that contain at least one Skinned Mesh component*/
+UCLASS(ClassGroup=(PerformanceCapture), meta=(BlueprintSpawnableComponent), Transient, DisplayName = "PerformanceCaptureVisualizer")
 class PERFORMANCECAPTUREWORKFLOW_API UPCapBoneVisualiser : public UInstancedStaticMeshComponent
 {
 	GENERATED_BODY()
@@ -29,11 +29,11 @@ public:
 	// Called every frame
 	virtual void TickComponent(float DeltaTime, ELevelTick TickType, FActorComponentTickFunction* ThisTickFunction) override;
 
-	/** Get all the transforms from the skelmesh. */
-	TArray<FTransform> GetJointTransforms(USkeletalMeshComponent* InSkelMeshComponent) const;
+	/** Get all the transforms from the Skinned Mesh component. */
+	TArray<FTransform> GetJointTransforms(USkinnedMeshComponent* InSkinnedMeshComponent) const;
 
 	/** Get all the bones' transforms for drawing a joint (point to point mesh). */
-	TArray<FTransform> GetBoneTransforms(USkeletalMeshComponent* InSkelMeshComponent) const;
+	TArray<FTransform> GetBoneTransforms(USkinnedMeshComponent* InSkinnedMeshComponent) const;
 	
 	/** Whether to draw joints along the length of the bone, or just at the pivot or scaled by the distance to the parent bone. */
 	UPROPERTY(EditAnywhere, BlueprintReadWrite, Category="PerformanceCapture|Visualization")
@@ -66,5 +66,5 @@ protected:
 
 private:
 	UPROPERTY()
-	TObjectPtr<USkeletalMeshComponent> SkelmeshComponent;
+	TObjectPtr<USkinnedMeshComponent> SkinnedMeshComponent;
 };
