[Performance Capture Workflow]
- Add missing delegate for when a Live Link subject is added. This is needed for Virtual Subjects

[RN] Added OnLiveLinkSubjectAdded delete to the PerformanceCaptureSubsystem

diff --git a/Engine/Plugins/VirtualProduction/PerformanceCaptureWorkflow/Source/PerformanceCaptureWorkflow/Private/PCapSubsystem.cpp b/Engine/Plugins/VirtualProduction/PerformanceCaptureWorkflow/Source/PerformanceCaptureWorkflow/Private/PCapSubsystem.cpp
--- a/Engine/Plugins/VirtualProduction/PerformanceCaptureWorkflow/Source/PerformanceCaptureWorkflow/Private/PCapSubsystem.cpp
+++ b/Engine/Plugins/VirtualProduction/PerformanceCaptureWorkflow/Source/PerformanceCaptureWorkflow/Private/PCapSubsystem.cpp
@@ -68,6 +68,7 @@ void UPerformanceCaptureSubsystem::OnEngineInitComplete()
 	{
 		ILiveLinkClient& LiveLinkClient = IModularFeatures::Get().GetModularFeature<ILiveLinkClient>(ILiveLinkClient::ModularFeatureName);
 		LiveLinkClient.OnLiveLinkSubjectStateChanged().AddUObject(this, &UPerformanceCaptureSubsystem::OnLiveLinkSubjectUpdated);
+		LiveLinkClient.OnLiveLinkSubjectAdded().AddUObject(this, &UPerformanceCaptureSubsystem::OnLiveLinkSubjectAdded);
 		LiveLinkClient.OnLiveLinkSubjectRemoved().AddUObject(this, &UPerformanceCaptureSubsystem::OnLiveLinkSubjectRemoved);
 		LiveLinkClient.OnLiveLinkSubjectEnabledChanged().AddUObject(this, &UPerformanceCaptureSubsystem::OnLiveLinkSubjectEnableChanged);
 	}
@@ -157,6 +158,11 @@ void UPerformanceCaptureSubsystem::OnLiveLinkSubjectRemoved(FLiveLinkSubjectKey
 	OnPCapLiveLinkSubjectRemoved.Broadcast(Subject);
 }
 
+void UPerformanceCaptureSubsystem::OnLiveLinkSubjectAdded(FLiveLinkSubjectKey Subject) const
+{
+	OnPCapLiveLinkSubjectAdded.Broadcast(Subject);
+}
+
 void UPerformanceCaptureSubsystem::OnLiveLinkSubjectEnableChanged(FLiveLinkSubjectKey Subject, bool NewEnabled) const
 {
 	OnPCapLiveLinkSubjectEnableChanged.Broadcast(Subject, NewEnabled);
diff --git a/Engine/Plugins/VirtualProduction/PerformanceCaptureWorkflow/Source/PerformanceCaptureWorkflow/Private/PCapSubsystem.h b/Engine/Plugins/VirtualProduction/PerformanceCaptureWorkflow/Source/PerformanceCaptureWorkflow/Private/PCapSubsystem.h
--- a/Engine/Plugins/VirtualProduction/PerformanceCaptureWorkflow/Source/PerformanceCaptureWorkflow/Private/PCapSubsystem.h
+++ b/Engine/Plugins/VirtualProduction/PerformanceCaptureWorkflow/Source/PerformanceCaptureWorkflow/Private/PCapSubsystem.h
@@ -127,6 +127,14 @@ public:
 	UPROPERTY(BlueprintAssignable, Category="Performance Capture|LiveLink")
 	FOnPCapLiveLinkSubjectUpdate OnPCapLiveLinkSubjectUpdate;
 
+	/**
+	* Multicast delegate that will execute when a Live Link subject is added.
+	*/
+	DECLARE_DYNAMIC_MULTICAST_DELEGATE_OneParam(FOnPCapLiveLinkSubjectAdded, FLiveLinkSubjectKey, Subject);
+	/** Editor Only*/
+	UPROPERTY(BlueprintAssignable, Category="Performance Capture|LiveLink")
+	FOnPCapLiveLinkSubjectAdded OnPCapLiveLinkSubjectAdded;
+
 	/**
 	* Multicast delegate that will execute when a Live Link subject is removed.
 	*/
@@ -185,6 +193,8 @@ private:
 
 	void OnLiveLinkSubjectUpdated(FLiveLinkSubjectKey Subject, ELiveLinkSubjectState State) const;
 
+	void OnLiveLinkSubjectAdded(FLiveLinkSubjectKey Subject) const;
+
 	void OnLiveLinkSubjectRemoved(FLiveLinkSubjectKey Subject) const;
 
 	void OnLiveLinkSubjectEnableChanged(FLiveLinkSubjectKey Subject, bool NewEnabled) const;
