[PerformanceCaptureWorkflow]
- Added dependency to Multiuser client plugin so can check if mocap manager is being used in concert session.
- Added dependency to Unreal Takes device plugin.

#rnx

diff --git a/Engine/Plugins/VirtualProduction/PerformanceCaptureWorkflow/PerformanceCaptureWorkflow.uplugin b/Engine/Plugins/VirtualProduction/PerformanceCaptureWorkflow/PerformanceCaptureWorkflow.uplugin
--- a/Engine/Plugins/VirtualProduction/PerformanceCaptureWorkflow/PerformanceCaptureWorkflow.uplugin
+++ b/Engine/Plugins/VirtualProduction/PerformanceCaptureWorkflow/PerformanceCaptureWorkflow.uplugin
@@ -1,7 +1,7 @@
 {
 	"FileVersion": 3,
 	"Version": 1,
-	"VersionName": "0.1",
+	"VersionName": "0.2",
 	"FriendlyName": "Performance Capture Workflow",
 	"Description": "Performance Capture In-Editor Workflow tools. Provides access to the Mocap Manager panel.",
 	"Category": "Animation",
@@ -74,6 +74,14 @@
 		{
 			"Name": "IKRig",
 			"Enabled": true
+		},
+		{
+			"Name": "MultiUserClient",
+			"Enabled": true
+		},
+		{
+			"Name": "LiveLinkUnrealDevice",
+			"Enabled": true
 		}
 	]
 }
\ No newline at end of file
