[Performance Capture Workflow]
- Added "<>" characters to those removed by filename and path santization functions
- Improved function display names

#rnx

diff --git a/Engine/Plugins/VirtualProduction/PerformanceCaptureWorkflow/Source/PerformanceCaptureWorkflow/Private/PCapBPFunctionLibrary.cpp b/Engine/Plugins/VirtualProduction/PerformanceCaptureWorkflow/Source/PerformanceCaptureWorkflow/Private/PCapBPFunctionLibrary.cpp
--- a/Engine/Plugins/VirtualProduction/PerformanceCaptureWorkflow/Source/PerformanceCaptureWorkflow/Private/PCapBPFunctionLibrary.cpp
+++ b/Engine/Plugins/VirtualProduction/PerformanceCaptureWorkflow/Source/PerformanceCaptureWorkflow/Private/PCapBPFunctionLibrary.cpp
@@ -8,7 +8,7 @@
 
 FString UPerformanceCaptureBPFunctionLibrary::SanitizeFileString(FString InString)
 {
-	FString BannedChars = "[]{}!\"$£%^&*()+=;:?\\/|'@#~ ";
+	FString BannedChars = "[]<>{}!\"$£%^&*()+=;:?\\/|'@#~ ";
 	FString CleanString = InString;
 
 	//BannedChars = BannedChars + FPaths::GetInvalidFileSystemChars(); //This probably causes duplicates but rather than consider each platform here, we'll rely on the core.
@@ -34,7 +34,7 @@ FString UPerformanceCaptureBPFunctionLibrary::SanitizePathString(FString InStrin
 	{
 		//This is the same as for file name strings but with allowed exception of "/" directory limiters
 		//TODO: currently windows only!
-		FString BannedChars = "[]{}!\"$£%^&*()+=;:?|'@#~ ";
+		FString BannedChars = "[]<>{}!\"$£%^&*()+=;:?|'@#~ ";
 		FString CleanString = InString;
 
 		//BannedChars = BannedChars + FPaths::GetInvalidFileSystemChars(); //This probably causes duplicates but rather than consider each platform here, we'll rely on the core.
diff --git a/Engine/Plugins/VirtualProduction/PerformanceCaptureWorkflow/Source/PerformanceCaptureWorkflow/Private/PCapBPFunctionLibrary.h b/Engine/Plugins/VirtualProduction/PerformanceCaptureWorkflow/Source/PerformanceCaptureWorkflow/Private/PCapBPFunctionLibrary.h
--- a/Engine/Plugins/VirtualProduction/PerformanceCaptureWorkflow/Source/PerformanceCaptureWorkflow/Private/PCapBPFunctionLibrary.h
+++ b/Engine/Plugins/VirtualProduction/PerformanceCaptureWorkflow/Source/PerformanceCaptureWorkflow/Private/PCapBPFunctionLibrary.h
@@ -22,7 +22,7 @@ public:
 	 * @param	InString String you wish to clean
 	 * @return	Cleaned string
 	 */
-	UFUNCTION(BlueprintCallable, Category = "Performance Capture|Workflow")
+	UFUNCTION(BlueprintCallable, Category = "Performance Capture|Workflow",meta=(DisplayName="Sanitize File Name String"))
 	static FString SanitizeFileString(FString InString);
 	
 	/**
@@ -30,7 +30,7 @@ public:
 	 * @param	InString String you wish to clean
 	 * @return	Cleaned string 
 	 */
-	UFUNCTION(BlueprintCallable, Category = "Performance Capture|Workflow")
+	UFUNCTION(BlueprintCallable, Category = "Performance Capture|Workflow",meta=(DisplayName="Sanitize File Path String"))
 	static FString SanitizePathString(FString InString);
 	
 	/** 
