[Performance Capture Workflow]
- Bughawk documentation string fix

diff --git a/Engine/Plugins/VirtualProduction/PerformanceCaptureWorkflow/Source/PerformanceCaptureWorkflow/Private/PCapBPFunctionLibrary.h b/Engine/Plugins/VirtualProduction/PerformanceCaptureWorkflow/Source/PerformanceCaptureWorkflow/Private/PCapBPFunctionLibrary.h
--- a/Engine/Plugins/VirtualProduction/PerformanceCaptureWorkflow/Source/PerformanceCaptureWorkflow/Private/PCapBPFunctionLibrary.h
+++ b/Engine/Plugins/VirtualProduction/PerformanceCaptureWorkflow/Source/PerformanceCaptureWorkflow/Private/PCapBPFunctionLibrary.h
@@ -18,7 +18,7 @@ class PERFORMANCECAPTUREWORKFLOW_API UPerformanceCaptureBPFunctionLibrary : publ
 	
 public:
 	/**
-	 * Returns a string stripped of the following characters []{}1!"$£%^&*()+=;:?/\|'@#~
+	 * Returns a string stripped of the following characters []<>{}1!"$£%^&*()+=;:?/\|'@#~
 	 * @param	InString String you wish to clean
 	 * @return	Cleaned string
 	 */
@@ -26,7 +26,7 @@ public:
 	static FString SanitizeFileString(FString InString);
 	
 	/**
-	 * Returns a string stripped of the following characters []{}1!"$£%^&*()+=;:?\|'@#~ Note filepath delimiters are not excluded
+	 * Returns a string stripped of the following characters []<>{}1!"$£%^&*()+=;:?\|'@#~ Note filepath delimiters are not excluded
 	 * @param	InString String you wish to clean
 	 * @return	Cleaned string 
 	 */
