[Performance Capture Workflow]
- Fixes incorrect serialization of multicast delegate by marking property transient

#rnx

diff --git a/Engine/Plugins/VirtualProduction/PerformanceCaptureWorkflow/Source/PerformanceCaptureWorkflow/Private/PCapDataTable.h b/Engine/Plugins/VirtualProduction/PerformanceCaptureWorkflow/Source/PerformanceCaptureWorkflow/Private/PCapDataTable.h
--- a/Engine/Plugins/VirtualProduction/PerformanceCaptureWorkflow/Source/PerformanceCaptureWorkflow/Private/PCapDataTable.h
+++ b/Engine/Plugins/VirtualProduction/PerformanceCaptureWorkflow/Source/PerformanceCaptureWorkflow/Private/PCapDataTable.h
@@ -19,7 +19,7 @@ public:
 	virtual ~UPCapDataTable() override;
 
 	DECLARE_DYNAMIC_MULTICAST_DELEGATE(FOnDatatableModified);
-	UPROPERTY(BlueprintAssignable, Category="Performance Capture|Database")
+	UPROPERTY(BlueprintAssignable, Transient, Category="Performance Capture|Database")
 	FOnDatatableModified OnDatatableModified;
 
 	/** Emit a callback each time the datatable is modified*/
