Performance Capture Workflow

-Code changes for composibe extensibility UE-351361
-Adds new dataasset type which user can create to define their ideal mocap manager workflow, optioning in and out which sections of the UI they wish to use.
-Use ini defined class to filter class of UMG widgets that are pickable in the new data asset type

diff --git a/Engine/Plugins/VirtualProduction/PerformanceCaptureWorkflow/Config/DefaultPerformanceCaptureWorkflow.ini b/Engine/Plugins/VirtualProduction/PerformanceCaptureWorkflow/Config/DefaultPerformanceCaptureWorkflow.ini
--- a/Engine/Plugins/VirtualProduction/PerformanceCaptureWorkflow/Config/DefaultPerformanceCaptureWorkflow.ini
+++ b/Engine/Plugins/VirtualProduction/PerformanceCaptureWorkflow/Config/DefaultPerformanceCaptureWorkflow.ini
@@ -4,6 +4,8 @@ bCanSaveAllSections=True
 [/Script/PerformanceCaptureWorkflow.PerformanceCaptureSettings]
 StageRoot=/PerformanceCaptureWorkflow/Stage/BP_StageRoot.BP_StageRoot_C
 ViewModelClass=/PerformanceCaptureWorkflow/Core/VM_PCapWorkflow.VM_PCapWorkflow_C
+MocapManagerDefaultWorkflow=/PerformanceCaptureWorkflow/Core/DA_MocapManagerDefaultWorkflow.DA_MocapManagerDefaultWorkflow
+WorkflowPhaseUIBaseClass=/PerformanceCaptureWorkflow/Core/EW_PCapPanelBase.EW_PCapPanelBase_C
 MocapManagerUI=/PerformanceCaptureWorkflow/Core/EW_MocapManager.EW_MocapManager
 DatabaseHelperClass=/PerformanceCaptureWorkflow/Database/EU_InternalDatabaseHelper.EU_InternalDatabaseHelper_C
 AllowedCaptureCharacterActorClasses=/Script/PerformanceCaptureCore.CaptureCharacter
diff --git a/Engine/Plugins/VirtualProduction/PerformanceCaptureWorkflow/Source/PerformanceCaptureWorkflow/Private/PCapSettings.h b/Engine/Plugins/VirtualProduction/PerformanceCaptureWorkflow/Source/PerformanceCaptureWorkflow/Private/PCapSettings.h
--- a/Engine/Plugins/VirtualProduction/PerformanceCaptureWorkflow/Source/PerformanceCaptureWorkflow/Private/PCapSettings.h
+++ b/Engine/Plugins/VirtualProduction/PerformanceCaptureWorkflow/Source/PerformanceCaptureWorkflow/Private/PCapSettings.h
@@ -11,6 +11,7 @@
 #include "UObject/Object.h"
 #include "PCapSettings.generated.h"
 
+class UPCapWorkflowCustomization;
 class APerformanceCaptureStageRoot;
 class UPCapSessionTemplate;
 class UDataTable;
@@ -64,6 +65,17 @@ public:
 	UPROPERTY(EditAnywhere, BlueprintReadOnly, Config, Category="UI", meta = (MetaClass = "/Script/Blutility.EditorUtilityWidgetBlueprint"))
 	TSoftObjectPtr<UEditorUtilityWidgetBlueprint> MocapManagerUI;
 
+	/** Default Workflow UI panels to use in the Mocap Manager panel. */
+	UPROPERTY(EditAnywhere, BlueprintReadOnly, Config, Category="UI", DisplayName = "Default Workflows", meta = ( tooltip="This Data Asset defines the UI panels you wish to have available in Mocap Manager"))
+	TSoftObjectPtr<UPCapWorkflowCustomization> MocapManagerDefaultWorkflow;
+
+	/** 
+	 * Editor Utility Widget class that will be used for each of the phase UI panels. Note this is not editable in the settings.
+	 * It can only be modified in the DefaultPerformanceCaptureWorkflow.ini
+	 */
+	UPROPERTY(BlueprintReadOnly, Config, NoClear, Category="UI", meta = (MetaClass = "/Script/Blutility.EditorUtilityWidgetBlueprint"))
+	TSoftClassPtr<UEditorUtilityWidget> WorkflowPhaseUIBaseClass;
+
 	/** Default session template to use when starting the Mocap Manager panel. */
 	UPROPERTY(BlueprintReadOnly, Config, EditDefaultsOnly, Category="Database", DisplayName = "Default Session Template", meta = ( tooltip="This Asset defines the folder structure and settings used in a performance capture session"))
 	TSoftObjectPtr<UPCapSessionTemplate> DefaultSessionTemplate;
@@ -76,7 +88,7 @@ public:
 	UPROPERTY(BlueprintReadOnly, Config, EditDefaultsOnly, Category = "Database|Internal Database Tables", DisplayName = "Session DataTable", meta = (RequiredAssetDataTags = "RowStructure=/Script/PerformanceCaptureWorkflow.PCapSessionRecord"))
 	TSoftObjectPtr<UPCapDataTable> SessionTable;
 
-	/** Proint to the datatable Mocap Manager will use for recording production data. */
+	/** Pointer to the datatable Mocap Manager will use for recording production data. */
 	UPROPERTY(BlueprintReadOnly, Config, EditDefaultsOnly, Category="Database|Internal Database Tables", DisplayName = "Production DataTable", meta = (RequiredAssetDataTags = "RowStructure=/Script/PerformanceCaptureWorkflow.PCapProductionRecord"))
 	TSoftObjectPtr<UPCapDataTable> ProductionTable;
 
diff --git a/Engine/Plugins/VirtualProduction/PerformanceCaptureWorkflow/Source/PerformanceCaptureWorkflow/Private/PCapWorkflowCustomization.cpp b/Engine/Plugins/VirtualProduction/PerformanceCaptureWorkflow/Source/PerformanceCaptureWorkflow/Private/PCapWorkflowCustomization.cpp
new file mode 100644
--- /dev/null
+++ b/Engine/Plugins/VirtualProduction/PerformanceCaptureWorkflow/Source/PerformanceCaptureWorkflow/Private/PCapWorkflowCustomization.cpp
@@ -0,0 +1,18 @@
+﻿// Copyright Epic Games, Inc. All Rights Reserved.
+
+#pragma once
+
+#include "PCapWorkflowCustomization.h"
+#include "PCapSettings.h"
+
+TArray<UClass*> UPCapWorkflowCustomization::GetAllowedPanelWidgetClasses()
+{
+	UPerformanceCaptureSettings* Settings = UPerformanceCaptureSettings::GetPerformanceCaptureSettings();
+	TArray<UClass*> AllowedClasses;
+	AllowedClasses.Reserve(1);
+	if (Settings != nullptr && !Settings->WorkflowPhaseUIBaseClass.IsNull())
+	{
+		AllowedClasses.Emplace(Settings->WorkflowPhaseUIBaseClass.LoadSynchronous());
+	}
+	return AllowedClasses;
+}
diff --git a/Engine/Plugins/VirtualProduction/PerformanceCaptureWorkflow/Source/PerformanceCaptureWorkflow/Private/PCapWorkflowCustomization.h b/Engine/Plugins/VirtualProduction/PerformanceCaptureWorkflow/Source/PerformanceCaptureWorkflow/Private/PCapWorkflowCustomization.h
new file mode 100644
--- /dev/null
+++ b/Engine/Plugins/VirtualProduction/PerformanceCaptureWorkflow/Source/PerformanceCaptureWorkflow/Private/PCapWorkflowCustomization.h
@@ -0,0 +1,60 @@
+﻿// Copyright Epic Games, Inc. All Rights Reserved.
+
+#pragma once
+
+#include "CoreMinimal.h"
+#include "EditorUtilityWidget.h"
+#include "Engine/DataAsset.h"
+#include "PCapWorkflowCustomization.generated.h"
+
+class UEditorUtilityWidget;
+
+USTRUCT(BlueprintType)
+struct FMocapManagerPhase
+{
+	GENERATED_BODY()
+
+	/**
+	* Name to use for this phase. It should be unique among all the other phase names
+	*/
+	UPROPERTY(EditAnywhere, BlueprintReadOnly, Category = "MocapManager")
+	FName PhaseName;
+	
+	/**
+	 * Icon to use on the phase button on the Mocap Manager UI
+	 */
+	UPROPERTY(EditAnywhere, BlueprintReadOnly, Category = "MocapManager")
+	TSoftObjectPtr<UTexture2D> Icon;
+
+	/**
+	 * Bool to determine if Mocap Manager should use this phase
+	 */
+	UPROPERTY(EditAnywhere, BlueprintReadOnly, Category = "MocapManager")
+	bool bIsEnabled = true;
+
+	/**
+	 * UI Panel to use for this phase. Must be of class defined by WorkflowPhaseUIBaseClass in DefaultPerformanceCaptureWorkflow.ini 
+	 */
+	UPROPERTY(EditAnywhere, BlueprintReadOnly, Category = "MocapManager", meta = (GetAllowedClasses = "/Script/PerformanceCaptureWorkflow.PCapWorkflowCustomization:GetAllowedPanelWidgetClasses"))
+	TSoftClassPtr<UEditorUtilityWidget> WorkflowPhaseUIPanelClass;
+	
+};
+
+/**
+ * Data asset class to let users build a series of phases in Mocap Manager
+ */
+UCLASS(Blueprintable, BlueprintType)
+class PERFORMANCECAPTUREWORKFLOW_API UPCapWorkflowCustomization : public UPrimaryDataAsset
+{
+	GENERATED_BODY()
+
+public:
+	UPROPERTY(EditAnywhere, BlueprintReadOnly, Category = "MocapManager")
+	TArray<FMocapManagerPhase> WorkflowPhases;
+
+private:
+	UFUNCTION()
+	static TArray<UClass*> GetAllowedPanelWidgetClasses();
+	
+};
+
