Performance Capture Workflow
- Added Icon for workflow customization data asset and custom color to match the rest of Mocap Manager assets.

#rnx

diff --git a/Engine/Plugins/VirtualProduction/PerformanceCaptureWorkflow/Source/PerformanceCaptureWorkflow/Private/PCapAssetDefinition.cpp b/Engine/Plugins/VirtualProduction/PerformanceCaptureWorkflow/Source/PerformanceCaptureWorkflow/Private/PCapAssetDefinition.cpp
--- a/Engine/Plugins/VirtualProduction/PerformanceCaptureWorkflow/Source/PerformanceCaptureWorkflow/Private/PCapAssetDefinition.cpp
+++ b/Engine/Plugins/VirtualProduction/PerformanceCaptureWorkflow/Source/PerformanceCaptureWorkflow/Private/PCapAssetDefinition.cpp
@@ -63,4 +63,17 @@ const FSlateBrush* UAssetDefinition_SessionTemplateAsset::GetThumbnailBrush(cons
 	return FPerformanceCaptureStyle::Get().GetBrush("ClassThumbnail.PCapSessionTemplate");
 }
 
+/*------------------------------------------------------------------------------
+	UPCapData Workflow Customization implementation.
+------------------------------------------------------------------------------*/
+const FSlateBrush* UAssetDefinition_PCapDataWorkflowCustomization::GetIconBrush(const FAssetData& InAssetData, const FName InClassName) const
+{
+	return FPerformanceCaptureStyle::Get().GetBrush("ClassIcon.PCapWorkflowCustomization");
+}
+
+const FSlateBrush* UAssetDefinition_PCapDataWorkflowCustomization::GetThumbnailBrush(const FAssetData& InAssetData, const FName InClassName) const
+{
+	return FPerformanceCaptureStyle::Get().GetBrush("ClassThumbnail.PCapWorkflowCustomization");
+}
+
 #undef LOCTEXT_NAMESPACE
\ No newline at end of file
diff --git a/Engine/Plugins/VirtualProduction/PerformanceCaptureWorkflow/Source/PerformanceCaptureWorkflow/Private/PCapAssetDefinition.h b/Engine/Plugins/VirtualProduction/PerformanceCaptureWorkflow/Source/PerformanceCaptureWorkflow/Private/PCapAssetDefinition.h
--- a/Engine/Plugins/VirtualProduction/PerformanceCaptureWorkflow/Source/PerformanceCaptureWorkflow/Private/PCapAssetDefinition.h
+++ b/Engine/Plugins/VirtualProduction/PerformanceCaptureWorkflow/Source/PerformanceCaptureWorkflow/Private/PCapAssetDefinition.h
@@ -6,6 +6,7 @@
 #include "PCapDataTable.h"
 #include "PCapDatabase.h"
 #include "PCapSessionTemplate.h"
+#include "PCapWorkflowCustomization.h"
 #include "Table/AssetDefinition_DataTable.h"
 
 #include "PCapAssetDefinition.generated.h"
@@ -84,6 +85,7 @@ public:
 	virtual TSoftClassPtr<UObject> GetAssetClass() const override { return UPCapPropDataAsset::StaticClass(); }
 	// UAssetDefinition End
 };
+
 UCLASS()
 class UAssetDefinition_SessionTemplateAsset : public UAssetDefinition_PCapDataAsset
 {
@@ -96,4 +98,18 @@ public:
 	virtual const FSlateBrush* GetThumbnailBrush(const FAssetData& InAssetData, const FName InClassName) const override;
 	virtual TSoftClassPtr<UObject> GetAssetClass() const override { return UPCapSessionTemplate::StaticClass(); }
 	// UAssetDefinition End
+};
+UCLASS()
+class UAssetDefinition_PCapDataWorkflowCustomization : public UAssetDefinition_DataAsset
+{
+	GENERATED_BODY()
+public:
+	
+	// UAssetDefinition Implementation
+	virtual FText GetAssetDisplayName() const override { return NSLOCTEXT("AssetDefinition", "PCapDataAsset", "PCap DataAsset"); }
+	virtual FLinearColor GetAssetColor() const override { return FLinearColor(FColor(161, 57, 191)); }
+	virtual const FSlateBrush* GetIconBrush(const FAssetData& InAssetData, const FName InClassName) const override;
+	virtual const FSlateBrush* GetThumbnailBrush(const FAssetData& InAssetData, const FName InClassName) const override;
+	virtual TSoftClassPtr<UObject> GetAssetClass() const override { return UPCapWorkflowCustomization::StaticClass(); }
+	// UAssetDefinition End
 };
\ No newline at end of file
diff --git a/Engine/Plugins/VirtualProduction/PerformanceCaptureWorkflow/Source/PerformanceCaptureWorkflow/Private/PerformanceCaptureStyle.cpp b/Engine/Plugins/VirtualProduction/PerformanceCaptureWorkflow/Source/PerformanceCaptureWorkflow/Private/PerformanceCaptureStyle.cpp
--- a/Engine/Plugins/VirtualProduction/PerformanceCaptureWorkflow/Source/PerformanceCaptureWorkflow/Private/PerformanceCaptureStyle.cpp
+++ b/Engine/Plugins/VirtualProduction/PerformanceCaptureWorkflow/Source/PerformanceCaptureWorkflow/Private/PerformanceCaptureStyle.cpp
@@ -51,7 +51,6 @@ TSharedRef< FSlateStyleSet > FPerformanceCaptureStyle::Create()
 	Style->Set("PerformanceCapture.MocapManagerTabIcon", new IMAGE_BRUSH_SVG(TEXT("ButtonIcon"), Icon20x20));
 
 	Style->Set("PerformanceCapture.MocapManagerTabIcon.Small", new IMAGE_BRUSH_SVG(TEXT("ButtonIcon"), Icon16x16));
-	
 
 	Style->Set("ClassThumbnail.PCapDataTable", new IMAGE_BRUSH_SVG(TEXT("DataTableThumbnail"), Icon64x64));
 
@@ -72,6 +71,10 @@ TSharedRef< FSlateStyleSet > FPerformanceCaptureStyle::Create()
 	Style->Set("ClassThumbnail.PCapSessionTemplate", new IMAGE_BRUSH_SVG(TEXT("PcapSessionTemplate_64"), Icon64x64));
 
 	Style->Set("ClassIcon.PCapSessionTemplate", new IMAGE_BRUSH_SVG(TEXT("PcapSessionTemplate_16"), Icon16x16));
+
+	Style->Set("ClassThumbnail.PCapWorkflowCustomization", new IMAGE_BRUSH_SVG(TEXT("WorkflowIcon"), Icon64x64));
+
+	Style->Set("ClassIcon.PCapWorkflowCustomization", new IMAGE_BRUSH_SVG(TEXT("WorkflowIcon"), Icon16x16));
 	
 	return Style;
 }
