Performance Capture Worflow
- Fix CIS issue arising from incorrect pragma use.

#rnx

diff --git a/Engine/Plugins/VirtualProduction/PerformanceCaptureWorkflow/Source/PerformanceCaptureWorkflow/Private/PCapWorkflowCustomization.cpp b/Engine/Plugins/VirtualProduction/PerformanceCaptureWorkflow/Source/PerformanceCaptureWorkflow/Private/PCapWorkflowCustomization.cpp
--- a/Engine/Plugins/VirtualProduction/PerformanceCaptureWorkflow/Source/PerformanceCaptureWorkflow/Private/PCapWorkflowCustomization.cpp
+++ b/Engine/Plugins/VirtualProduction/PerformanceCaptureWorkflow/Source/PerformanceCaptureWorkflow/Private/PCapWorkflowCustomization.cpp
@@ -1,7 +1,5 @@
 ﻿// Copyright Epic Games, Inc. All Rights Reserved.
 
-#pragma once
-
 #include "PCapWorkflowCustomization.h"
 #include "PCapSettings.h"
 
