Mocap Manager
- Added date and time of recording finalisation to the takes record struct
- BP changes to support writing this new field into a takes datatable

[RN] Added time created to takes record
#virtualized
#Synced-CL 48712650

diff --git a/Engine/Plugins/VirtualProduction/PerformanceCaptureWorkflow/Source/PerformanceCaptureWorkflow/Private/PCapDatabase.h b/Engine/Plugins/VirtualProduction/PerformanceCaptureWorkflow/Source/PerformanceCaptureWorkflow/Private/PCapDatabase.h
--- a/Engine/Plugins/VirtualProduction/PerformanceCaptureWorkflow/Source/PerformanceCaptureWorkflow/Private/PCapDatabase.h
+++ b/Engine/Plugins/VirtualProduction/PerformanceCaptureWorkflow/Source/PerformanceCaptureWorkflow/Private/PCapDatabase.h
@@ -11,7 +11,6 @@
 #include "Engine/DataTable.h"
 #include "Engine/DataAsset.h"
 #include "Misc/DateTime.h"
-#include "UObject/ConstructorHelpers.h"
 #include "PCapDatabase.generated.h"
 
 class UPCapPropDataAsset;
@@ -90,8 +89,14 @@ public:
 	UPROPERTY(Transient, BlueprintReadWrite, Category="Flag")
 	bool bFlag = false;
 	
+	/** The date and time the sequence was recorded. This is the local machine
+	 * UTC at the moment recording is finalised
+	 */
+	UPROPERTY(VisibleAnywhere, BlueprintReadWrite, Category="Performance Capture Take")
+	FDateTime DateTimeCreated;
+	
 	/** The level sequence recorded. */
-	UPROPERTY(EditAnywhere, BlueprintReadWrite, Category="Performance Capture Take")
+	UPROPERTY(VisibleAnywhere, BlueprintReadWrite, Category="Performance Capture Take")
 	TSoftObjectPtr<ULevelSequence> RecordedTake;
 	
 	/** Seconds Duration. */
@@ -115,7 +120,7 @@ public:
 	FTimecode EndTimecode;
 
 	/** Transform of the stage root when this recording was made. */
-	UPROPERTY(EditAnywhere, BlueprintReadWrite, Category="Performance Capture Take")
+	UPROPERTY(VisibleAnywhere, BlueprintReadWrite, Category="Performance Capture Take")
 	FTransform MocapStageRootTransform;
 
 	/** Status of Take - Thumbs Up, Thumbs Down or None. */
