Performance Capture
- Added a scalar to the prop component allowing for scale of a prop's motion from the origin of the motion capture system, AKA global scale of prop.
- Does not affect the scale of the mesh - only affects the location.

#virtualized
#Synced-CL 48791601

diff --git a/Engine/Plugins/VirtualProduction/PerformanceCaptureWorkflow/Source/PerformanceCaptureWorkflowRuntime/Private/AnimNode_LiveLinkProp.cpp b/Engine/Plugins/VirtualProduction/PerformanceCaptureWorkflow/Source/PerformanceCaptureWorkflowRuntime/Private/AnimNode_LiveLinkProp.cpp
--- a/Engine/Plugins/VirtualProduction/PerformanceCaptureWorkflow/Source/PerformanceCaptureWorkflowRuntime/Private/AnimNode_LiveLinkProp.cpp
+++ b/Engine/Plugins/VirtualProduction/PerformanceCaptureWorkflow/Source/PerformanceCaptureWorkflowRuntime/Private/AnimNode_LiveLinkProp.cpp
@@ -28,25 +28,32 @@ void FAnimNode_LiveLinkProp::BuildPoseFromAnimData(const FLiveLinkSubjectFrameDa
 	//Get bones names
 	const TArray<FName>& SourceBoneNames = SkeletonData->GetBoneNames();
 
-	//Iterate through bones in LiveLink Data
-	for (int32 i = 0; i < SourceBoneNames.Num(); ++i)
+	if (SkeletonData->BoneNames.Num() == FrameData->Transforms.Num())
 	{
-		FName BoneName = SourceBoneNames[i];
+		//Iterate through bones in LiveLink Data
+		for (int32 i = 0; i < SourceBoneNames.Num(); ++i)
+		{
+			FName BoneName = SourceBoneNames[i];
+		
+			FTransform LiveLinkTransform = FrameData->Transforms[i];
+			LiveLinkTransform.SetLocation(FrameData->Transforms[i].GetLocation() * GlobalScale);
 
-		FTransform BoneTransform = OffsetTransform*FrameData->Transforms[i];
+			FTransform BoneTransform = OffsetTransform * LiveLinkTransform ;
 
-		BoneTransform.SetLocation(BoneTransform.GetLocation()+DynamicConstraintOffset);
+			BoneTransform.SetLocation(BoneTransform.GetLocation() + DynamicConstraintOffset);
 
-		int32 MeshIndex = Output.Pose.GetBoneContainer().GetPoseBoneIndexForBoneName(BoneName);
-		if (MeshIndex != INDEX_NONE)
-		{
-			FCompactPoseBoneIndex CompactPoseBoneIndex = Output.Pose.GetBoneContainer().MakeCompactPoseIndex(FMeshPoseBoneIndex(MeshIndex));
-			if(CompactPoseBoneIndex != INDEX_NONE)
+			int32 MeshIndex = Output.Pose.GetBoneContainer().GetPoseBoneIndexForBoneName(BoneName);
+			if (MeshIndex != INDEX_NONE)
 			{
-				Output.Pose[CompactPoseBoneIndex] = BoneTransform;
+				FCompactPoseBoneIndex CompactPoseBoneIndex = Output.Pose.GetBoneContainer().MakeCompactPoseIndex(FMeshPoseBoneIndex(MeshIndex));
+				if(CompactPoseBoneIndex != INDEX_NONE)
+				{
+					Output.Pose[CompactPoseBoneIndex] = BoneTransform;
+				}
 			}
 		}
 	}
+	
 	CachedDeltaTime = 0.0f;
 }
 
diff --git a/Engine/Plugins/VirtualProduction/PerformanceCaptureWorkflow/Source/PerformanceCaptureWorkflowRuntime/Private/AnimNode_LiveLinkProp.h b/Engine/Plugins/VirtualProduction/PerformanceCaptureWorkflow/Source/PerformanceCaptureWorkflowRuntime/Private/AnimNode_LiveLinkProp.h
--- a/Engine/Plugins/VirtualProduction/PerformanceCaptureWorkflow/Source/PerformanceCaptureWorkflowRuntime/Private/AnimNode_LiveLinkProp.h
+++ b/Engine/Plugins/VirtualProduction/PerformanceCaptureWorkflow/Source/PerformanceCaptureWorkflowRuntime/Private/AnimNode_LiveLinkProp.h
@@ -39,6 +39,10 @@ public:
 	/* Transform to apply local space to offset to the incoming Live Link data. */
 	UPROPERTY(EditAnywhere, Category = "Performance Capture|Prop")
 	FTransform OffsetTransform = FTransform::Identity;
+	
+	/* Global scale to the location*/
+	UPROPERTY(EditAnywhere, Category = "Performance Capture|Prop")
+	float GlobalScale = 1.0f;
 
 	/* Location to apply in parent bone space to offset the incoming Live Link data. */
 	UPROPERTY(EditAnywhere, Category = "Performance Capture|Prop")
diff --git a/Engine/Plugins/VirtualProduction/PerformanceCaptureWorkflow/Source/PerformanceCaptureWorkflowRuntime/Private/PCapPropComponent.cpp b/Engine/Plugins/VirtualProduction/PerformanceCaptureWorkflow/Source/PerformanceCaptureWorkflowRuntime/Private/PCapPropComponent.cpp
--- a/Engine/Plugins/VirtualProduction/PerformanceCaptureWorkflow/Source/PerformanceCaptureWorkflowRuntime/Private/PCapPropComponent.cpp
+++ b/Engine/Plugins/VirtualProduction/PerformanceCaptureWorkflow/Source/PerformanceCaptureWorkflowRuntime/Private/PCapPropComponent.cpp
@@ -102,6 +102,21 @@ void UPCapPropComponent::SetOffsetTransform(FTransform NewOffset)
 	}
 }
 
+void UPCapPropComponent::SetGlobalScale(float NewGlobalScale)
+{
+	GlobalScale = NewGlobalScale;
+	
+	//If there is a Skeletal Mesh controlled by an Anim Instance pass the new global scale
+	TObjectPtr<USkeletalMeshComponent> ControlledMesh = Cast<USkeletalMeshComponent>(ControlledComponent.GetComponent(GetOwner()));
+	if(IsValid(ControlledMesh))
+	{
+		if(UPCapPropLiveLinkAnimInstance* LiveLinkAnimInstance = Cast<UPCapPropLiveLinkAnimInstance>(ControlledMesh->GetAnimInstance()))
+		{
+			LiveLinkAnimInstance->SetGlobalScale(GlobalScale);
+		}
+	}
+}
+
 TArray<FName> UPCapPropComponent::GetAttachBones() const
 {
 	TArray<FName> FoundIKChainBones;
@@ -177,25 +192,27 @@ void UPCapPropComponent::DestroyComponent(bool bPromoteChildren)
 #if WITH_EDITOR
 void UPCapPropComponent::PostEditChangeProperty(FPropertyChangedEvent& PropertyChangedEvent)
 {
-	if(ControlledComponent.GetComponent(GetOwner()))
+	if(PropertyChangedEvent.Property != nullptr)
 	{
-		Cast<USceneComponent>(ControlledComponent.GetComponent(GetOwner()))->SetMobility(EComponentMobility::Movable);
-	}
-	if(PropertyChangedEvent.Property->GetFName() == GET_MEMBER_NAME_CHECKED(UPCapPropComponent, OffsetTransform))
-	{
-		SetOffsetTransform(OffsetTransform);
-	}
-	if(PropertyChangedEvent.Property->GetFName() == GET_MEMBER_NAME_CHECKED(UPCapPropComponent, ControlledComponent))
-	{
-		ResetAnimInstance();
-		bIsDirty = true;
-	}
-	if(PropertyChangedEvent.Property->GetFName() == GET_MEMBER_NAME_CHECKED(UPCapPropComponent, SubjectName))
-	{
-		ResetAnimInstance();
-		bIsDirty = true;
+		if(ControlledComponent.GetComponent(GetOwner()))
+		{
+			Cast<USceneComponent>(ControlledComponent.GetComponent(GetOwner()))->SetMobility(EComponentMobility::Movable);
+		}
+		if(PropertyChangedEvent.Property->GetFName() == GET_MEMBER_NAME_CHECKED(UPCapPropComponent, OffsetTransform))
+		{
+			SetOffsetTransform(OffsetTransform);
+		}
+		if(PropertyChangedEvent.Property->GetFName() == GET_MEMBER_NAME_CHECKED(UPCapPropComponent, ControlledComponent))
+		{
+			ResetAnimInstance();
+			bIsDirty = true;
+		}
+		if(PropertyChangedEvent.Property->GetFName() == GET_MEMBER_NAME_CHECKED(UPCapPropComponent, SubjectName))
+		{
+			ResetAnimInstance();
+			bIsDirty = true;
+		}
 	}
-
 	Super::PostEditChangeProperty(PropertyChangedEvent);
 }
 #endif
@@ -219,6 +236,9 @@ void UPCapPropComponent::TickComponent(float DeltaTime, ELevelTick TickType, FAc
 			bIsControlledBySequencer = true;
 		}
 	}
+	
+	FTransform GlobalScaleTransform = FTransform::Identity;
+	GlobalScaleTransform.SetScale3D(FVector(GlobalScale, GlobalScale, GlobalScale));
 
 	if(!bIsControlledBySequencer)
 	{
@@ -243,15 +263,14 @@ void UPCapPropComponent::TickComponent(float DeltaTime, ELevelTick TickType, FAc
 
 				if (bUseDynamicConstraint)
 				{
-					FinalTransform.SetComponents(CachedLiveLinkTransform.GetRotation(),CachedLiveLinkTransform.GetLocation()+DynamicConstraintOffset, CachedLiveLinkTransform.GetScale3D());
+					FinalTransform.SetComponents(CachedLiveLinkTransform.GetRotation(),(CachedLiveLinkTransform.GetLocation() * GlobalScale) + DynamicConstraintOffset, CachedLiveLinkTransform.GetScale3D());
 				}
 				else
 				{
-					FinalTransform.SetComponents(CachedLiveLinkTransform.GetRotation(),CachedLiveLinkTransform.GetLocation(), CachedLiveLinkTransform.GetScale3D());
+					FinalTransform.SetComponents(CachedLiveLinkTransform.GetRotation(),(CachedLiveLinkTransform.GetLocation() * GlobalScale), CachedLiveLinkTransform.GetScale3D());
 				}
 				
 				GetControlledComponent()->SetRelativeTransform(OffsetTransform * FinalTransform, false, nullptr, ETeleportType::None);
-
 			}
 			else if(LiveLinkClient.GetSubjectRole_AnyThread(SubjectName) == ULiveLinkTransformRole::StaticClass())
 			{
@@ -263,11 +282,11 @@ void UPCapPropComponent::TickComponent(float DeltaTime, ELevelTick TickType, FAc
 				
 				if (bUseDynamicConstraint)
 				{
-					FinalTransform.SetComponents(CachedLiveLinkTransform.GetRotation(),CachedLiveLinkTransform.GetLocation()+DynamicConstraintOffset, CachedLiveLinkTransform.GetScale3D());
+					FinalTransform.SetComponents(CachedLiveLinkTransform.GetRotation(),(CachedLiveLinkTransform.GetLocation() * GlobalScale) + DynamicConstraintOffset, CachedLiveLinkTransform.GetScale3D());
 				}
 				else
 				{
-					FinalTransform.SetComponents(CachedLiveLinkTransform.GetRotation(),CachedLiveLinkTransform.GetLocation(), CachedLiveLinkTransform.GetScale3D());
+					FinalTransform.SetComponents(CachedLiveLinkTransform.GetRotation(),(CachedLiveLinkTransform.GetLocation() * GlobalScale), CachedLiveLinkTransform.GetScale3D());
 				}
 				
 				GetControlledComponent()->SetRelativeTransform(OffsetTransform * FinalTransform, false, nullptr, ETeleportType::None);
@@ -292,11 +311,11 @@ void UPCapPropComponent::TickComponent(float DeltaTime, ELevelTick TickType, FAc
 
 				if (bUseDynamicConstraint)
 				{
-					FinalTransform.SetComponents(CachedLiveLinkTransform.GetRotation(),CachedLiveLinkTransform.GetLocation()+DynamicConstraintOffset, CachedLiveLinkTransform.GetScale3D());
+					FinalTransform.SetComponents(CachedLiveLinkTransform.GetRotation(),(CachedLiveLinkTransform.GetLocation() * GlobalScale) + DynamicConstraintOffset, CachedLiveLinkTransform.GetScale3D());
 				}
 				else
 				{
-					FinalTransform.SetComponents(CachedLiveLinkTransform.GetRotation(),CachedLiveLinkTransform.GetLocation(), CachedLiveLinkTransform.GetScale3D());
+					FinalTransform.SetComponents(CachedLiveLinkTransform.GetRotation(),(CachedLiveLinkTransform.GetLocation() * GlobalScale), CachedLiveLinkTransform.GetScale3D());
 				}
 				
 				GetControlledComponent()->SetRelativeTransform(OffsetTransform * FinalTransform, false, nullptr, ETeleportType::None);
@@ -324,6 +343,9 @@ void UPCapPropComponent::InitiateAnimation()
 	TObjectPtr<USkeletalMeshComponent> ControlledMesh = Cast<USkeletalMeshComponent>(ControlledComponent.GetComponent(GetOwner()));
 	ILiveLinkClient& LiveLinkClient = IModularFeatures::Get().GetModularFeature<ILiveLinkClient>(ILiveLinkClient::ModularFeatureName);
 	
+	FTransform GlobalScaleTransform = FTransform::Identity;
+	GlobalScaleTransform.SetScale3D(FVector(GlobalScale, GlobalScale, GlobalScale));
+	
 	//If we are controlling a Skeletal Mesh and the Live Link data is of the animation role, initialize animation instance
 	if(IsValid(ControlledMesh) && LiveLinkClient.GetSubjectRole_AnyThread(SubjectName) == ULiveLinkAnimationRole::StaticClass())
 	{
@@ -340,6 +362,8 @@ void UPCapPropComponent::InitiateAnimation()
 				AnimInstance->SetSubject(SubjectName);
 				AnimInstance->EnableLiveLinkEvaluation(bEvaluateLiveLink);
 				AnimInstance->SetOffsetTransform(OffsetTransform);
+				AnimInstance->SetGlobalScale(GlobalScale);
+				AnimInstance->SetDynamicConstraintVector(FVector::ZeroVector);
 				bIsDirty=false;
 			}
 		}
diff --git a/Engine/Plugins/VirtualProduction/PerformanceCaptureWorkflow/Source/PerformanceCaptureWorkflowRuntime/Private/PCapPropLiveLinkAnimInstance.h b/Engine/Plugins/VirtualProduction/PerformanceCaptureWorkflow/Source/PerformanceCaptureWorkflowRuntime/Private/PCapPropLiveLinkAnimInstance.h
--- a/Engine/Plugins/VirtualProduction/PerformanceCaptureWorkflow/Source/PerformanceCaptureWorkflowRuntime/Private/PCapPropLiveLinkAnimInstance.h
+++ b/Engine/Plugins/VirtualProduction/PerformanceCaptureWorkflow/Source/PerformanceCaptureWorkflowRuntime/Private/PCapPropLiveLinkAnimInstance.h
@@ -67,6 +67,12 @@ public:
 		GetProxyOnGameThread<FPCapPropLiveLinkAnimInstanceProxy>().PoseNode.DynamicConstraintOffset = DynamicOffset;
 	}
 	
+	UFUNCTION(BlueprintCallable, Category = "Performance Capture|Prop")
+	void SetGlobalScale(float GlobalScale)
+	{
+		GetProxyOnGameThread<FPCapPropLiveLinkAnimInstanceProxy>().PoseNode.GlobalScale = GlobalScale;
+	}
+	
 	//Not exposed to Blueprint
 	bool GetEnableLiveLinkEvaluation()
 	{
diff --git a/Engine/Plugins/VirtualProduction/PerformanceCaptureWorkflow/Source/PerformanceCaptureWorkflowRuntime/Public/PCapPropComponent.h b/Engine/Plugins/VirtualProduction/PerformanceCaptureWorkflow/Source/PerformanceCaptureWorkflowRuntime/Public/PCapPropComponent.h
--- a/Engine/Plugins/VirtualProduction/PerformanceCaptureWorkflow/Source/PerformanceCaptureWorkflowRuntime/Public/PCapPropComponent.h
+++ b/Engine/Plugins/VirtualProduction/PerformanceCaptureWorkflow/Source/PerformanceCaptureWorkflowRuntime/Public/PCapPropComponent.h
@@ -53,6 +53,13 @@ public:
 	 */
 	UPROPERTY(BlueprintReadOnly, EditAnywhere, Category = "Performance Capture")
 	FTransform OffsetTransform = FTransform::Identity;
+	
+	/**
+	 * Scale the transform from the parent pivot (usually the motion capture stage origin point).
+	 * Scale is only applied to the location, not the scale of the driven mesh. 
+	 */
+	UPROPERTY(BlueprintReadOnly, EditAnywhere, Category = "Performance Capture")
+	float GlobalScale = 1.0f;
 
 	/**
 	 * Scene Component driven by the LiveLink data. Defaults to Root component of the owning actor.
@@ -106,6 +113,13 @@ public:
 	*/
 	UFUNCTION(BlueprintCallable, Category ="Performance Capture|Prop")
 	void SetOffsetTransform(FTransform NewOffset);
+	
+	/**
+	* Update the scale factor applied to the incoming Live Link poses 
+	* @param NewGlobalScale Global Scale.
+	*/
+	UFUNCTION(BlueprintCallable, Category ="Performance Capture|Prop")
+	void SetGlobalScale(float NewGlobalScale);
 
 	/**
 	* Enable/Disable use of the Dynamic Constraint offset vector. 
