[Mocap Manager]
- Changed the way the menu entry is registered so it now appears in alphabetical order
- removed redundant call to ToolMenus

#Synced-CL 49587344

diff --git a/Engine/Plugins/VirtualProduction/PerformanceCaptureWorkflow/Source/PerformanceCaptureWorkflow/PerformanceCaptureWorkflow.Build.cs b/Engine/Plugins/VirtualProduction/PerformanceCaptureWorkflow/Source/PerformanceCaptureWorkflow/PerformanceCaptureWorkflow.Build.cs
--- a/Engine/Plugins/VirtualProduction/PerformanceCaptureWorkflow/Source/PerformanceCaptureWorkflow/PerformanceCaptureWorkflow.Build.cs
+++ b/Engine/Plugins/VirtualProduction/PerformanceCaptureWorkflow/Source/PerformanceCaptureWorkflow/PerformanceCaptureWorkflow.Build.cs
@@ -52,7 +52,8 @@ public class PerformanceCaptureWorkflow : ModuleRules
 				"TakesCore",
 				"SubobjectDataInterface",
 				"NamingTokens",
-				"PerformanceCaptureWorkflowRuntime"
+				"PerformanceCaptureWorkflowRuntime", 
+				"WorkspaceMenuStructure"
 			}
 			);
 	}
diff --git a/Engine/Plugins/VirtualProduction/PerformanceCaptureWorkflow/Source/PerformanceCaptureWorkflow/Private/PerformanceCapture.cpp b/Engine/Plugins/VirtualProduction/PerformanceCaptureWorkflow/Source/PerformanceCaptureWorkflow/Private/PerformanceCapture.cpp
--- a/Engine/Plugins/VirtualProduction/PerformanceCaptureWorkflow/Source/PerformanceCaptureWorkflow/Private/PerformanceCapture.cpp
+++ b/Engine/Plugins/VirtualProduction/PerformanceCaptureWorkflow/Source/PerformanceCaptureWorkflow/Private/PerformanceCapture.cpp
@@ -15,6 +15,8 @@
 #include "Widgets/Layout/SBox.h"
 #include "Widgets/Text/STextBlock.h"
 #include "ToolMenus.h"
+#include "WorkspaceMenuStructure.h"
+#include "WorkspaceMenuStructureModule.h"
 
 static const FName PerformanceCapturePanelTabName("PerformanceCaptureTab");
 
@@ -42,13 +44,12 @@ void FPerformanceCaptureModule::StartupModule()
 		FPerformanceCaptureCommands::Get().OpenPluginWindow,
 		FExecuteAction::CreateRaw(this, &FPerformanceCaptureModule::PluginButtonClicked),
 		FCanExecuteAction());
-
-	UToolMenus::RegisterStartupCallback(FSimpleMulticastDelegate::FDelegate::CreateRaw(this, &FPerformanceCaptureModule::RegisterMenus));
-	
+		
 	FGlobalTabmanager::Get()->RegisterNomadTabSpawner(PerformanceCapturePanelTabName, FOnSpawnTab::CreateRaw(this, &FPerformanceCaptureModule::OnSpawnMocapManager))
 		.SetDisplayName(NSLOCTEXT("PerformanceCapture", "MocapManagerTabTitle", "Mocap Manager"))
 		.SetTooltipText(NSLOCTEXT("PerformanceCapture", "PerformanceCaptureTooltipText", "Open the Mocap Manager tab"))
-		.SetMenuType(ETabSpawnerMenuType::Hidden)
+		.SetGroup(WorkspaceMenu::GetMenuStructure().GetLevelEditorVirtualProductionCategory())
+		.SetMenuType(ETabSpawnerMenuType::Enabled)
 		.SetIcon(FSlateIcon(FPerformanceCaptureStyle::GetStyleSetName(), "PerformanceCapture.MocapManagerTabIcon", "PerformanceCapture.MocapManagerTabIcon.Small")
 		);
 
@@ -85,7 +86,7 @@ TSharedRef<SDockTab> FPerformanceCaptureModule::OnSpawnMocapManager(const FSpawn
 	
 	if(Settings->MocapManagerUI.IsValid())
 	{
-		UEditorUtilityWidgetBlueprint* MocapManagerEW = LoadObject<UEditorUtilityWidgetBlueprint>(NULL, *Settings->MocapManagerUI.ToString(), NULL, LOAD_None, NULL);
+		UEditorUtilityWidgetBlueprint* MocapManagerEW = LoadObject<UEditorUtilityWidgetBlueprint>(nullptr, *Settings->MocapManagerUI.ToString(), nullptr, LOAD_None, nullptr);
 	
 		TSharedRef<SDockTab> TabWidget = MocapManagerEW->SpawnEditorUITab(SpawnTabArgs);
 		
@@ -121,31 +122,6 @@ void FPerformanceCaptureModule::PluginButtonClicked()
 	FGlobalTabmanager::Get()->TryInvokeTab(PerformanceCapturePanelTabName);
 }
 
-void FPerformanceCaptureModule::RegisterMenus()
-{
-	FToolMenuOwnerScoped OwnerScoped(this);
-
-	{
-		UToolMenu* Menu = UToolMenus::Get()->ExtendMenu("LevelEditor.MainMenu.Window.VirtualProduction");
-		{
-			FToolMenuSection& Section = Menu->FindOrAddSection("VirtualProduction");
-			Section.AddMenuEntryWithCommandList(FPerformanceCaptureCommands::Get().OpenPluginWindow, PluginCommands);
-		}
-	}
-
-	{
-		UToolMenu* ToolbarMenu = UToolMenus::Get()->ExtendMenu("LevelEditor.LevelEditorToolBar");
-		{
-			FToolMenuSection& Section = ToolbarMenu->FindOrAddSection("Settings");
-			{
-				FToolMenuEntry& Entry = Section.AddEntry(FToolMenuEntry::InitToolBarButton(FPerformanceCaptureCommands::Get().OpenPluginWindow));
-				Entry.SetCommandList(PluginCommands);
-			}
-		}
-	}
-}
-
-
 #undef LOCTEXT_NAMESPACE
 	
 IMPLEMENT_MODULE(FPerformanceCaptureModule, PerformanceCaptureWorkflow)
diff --git a/Engine/Plugins/VirtualProduction/PerformanceCaptureWorkflow/Source/PerformanceCaptureWorkflow/Public/PerformanceCapture.h b/Engine/Plugins/VirtualProduction/PerformanceCaptureWorkflow/Source/PerformanceCaptureWorkflow/Public/PerformanceCapture.h
--- a/Engine/Plugins/VirtualProduction/PerformanceCaptureWorkflow/Source/PerformanceCaptureWorkflow/Public/PerformanceCapture.h
+++ b/Engine/Plugins/VirtualProduction/PerformanceCaptureWorkflow/Source/PerformanceCaptureWorkflow/Public/PerformanceCapture.h
@@ -26,9 +26,6 @@ public:
 	UE_API static const FLazyName MessageLogName; // "PerformanceCapture"
 
 private:
-
-	void RegisterMenus();
-
 	TSharedRef<class SDockTab> OnSpawnMocapManager(const class FSpawnTabArgs& SpawnTabArgs);
 	
 	TSharedPtr<class FUICommandList> PluginCommands;
