Mocap Manager - Adds multicast delegate to the workflow customization so that the UI can listen for changes to the customization and refresh itself.
- Updated the mocap manager widget to use the new delegate/event dispatcher

[RN] PCapWorkflow customization now has an event dispatcher triggered by post-edit so bound widgets can listen and react.
#virtualized
#Synced-CL 50249092

diff --git a/Engine/Plugins/VirtualProduction/PerformanceCaptureWorkflow/Source/PerformanceCaptureWorkflow/Private/PCapWorkflowCustomization.cpp b/Engine/Plugins/VirtualProduction/PerformanceCaptureWorkflow/Source/PerformanceCaptureWorkflow/Private/PCapWorkflowCustomization.cpp
--- a/Engine/Plugins/VirtualProduction/PerformanceCaptureWorkflow/Source/PerformanceCaptureWorkflow/Private/PCapWorkflowCustomization.cpp
+++ b/Engine/Plugins/VirtualProduction/PerformanceCaptureWorkflow/Source/PerformanceCaptureWorkflow/Private/PCapWorkflowCustomization.cpp
@@ -14,3 +14,16 @@ TArray<UClass*> UPCapWorkflowCustomization::GetAllowedPanelWidgetClasses()
 	}
 	return AllowedClasses;
 }
+
+#if WITH_EDITOR
+void UPCapWorkflowCustomization::PostEditChangeProperty(FPropertyChangedEvent& PropertyChangedEvent)
+{
+	Super::PostEditChangeProperty(PropertyChangedEvent);
+	
+	if(PropertyChangedEvent.Property != nullptr)
+	{
+		OnPCapCustomizationModified.Broadcast();
+	}
+}
+#endif
+
diff --git a/Engine/Plugins/VirtualProduction/PerformanceCaptureWorkflow/Source/PerformanceCaptureWorkflow/Private/PCapWorkflowCustomization.h b/Engine/Plugins/VirtualProduction/PerformanceCaptureWorkflow/Source/PerformanceCaptureWorkflow/Private/PCapWorkflowCustomization.h
--- a/Engine/Plugins/VirtualProduction/PerformanceCaptureWorkflow/Source/PerformanceCaptureWorkflow/Private/PCapWorkflowCustomization.h
+++ b/Engine/Plugins/VirtualProduction/PerformanceCaptureWorkflow/Source/PerformanceCaptureWorkflow/Private/PCapWorkflowCustomization.h
@@ -51,10 +51,23 @@ class PERFORMANCECAPTUREWORKFLOW_API UPCapWorkflowCustomization : public UPrimar
 public:
 	UPROPERTY(EditAnywhere, BlueprintReadOnly, Category = "MocapManager")
 	TArray<FMocapManagerPhase> WorkflowPhases;
+	
+	DECLARE_DYNAMIC_MULTICAST_DELEGATE(FOnPCapCustomizationModified);
+	UPROPERTY(BlueprintAssignable, Transient, Category="Performance Capture|Customization")
+	FOnPCapCustomizationModified OnPCapCustomizationModified;
+	
+	//~ UObject interface
+#if WITH_EDITOR
+	virtual void PostEditChangeProperty(FPropertyChangedEvent& PropertyChangedEvent) override;
+
+#endif
+	//~ End UObject interface
 
 private:
 	UFUNCTION()
 	static TArray<UClass*> GetAllowedPanelWidgetClasses();
 	
+
+	
 };
 
