Mocap Manager - Exposed Take Recorder clock source in the session template, viewmodel and the Record Manager UI.

#virtualized
#Synced-CL 50712830
#Synced-CL 50712830

diff --git a/Engine/Plugins/VirtualProduction/PerformanceCaptureWorkflow/Source/PerformanceCaptureWorkflow/PerformanceCaptureWorkflow.Build.cs b/Engine/Plugins/VirtualProduction/PerformanceCaptureWorkflow/Source/PerformanceCaptureWorkflow/PerformanceCaptureWorkflow.Build.cs
--- a/Engine/Plugins/VirtualProduction/PerformanceCaptureWorkflow/Source/PerformanceCaptureWorkflow/PerformanceCaptureWorkflow.Build.cs
+++ b/Engine/Plugins/VirtualProduction/PerformanceCaptureWorkflow/Source/PerformanceCaptureWorkflow/PerformanceCaptureWorkflow.Build.cs
@@ -53,7 +53,8 @@ public class PerformanceCaptureWorkflow : ModuleRules
 				"SubobjectDataInterface",
 				"NamingTokens",
 				"PerformanceCaptureWorkflowRuntime", 
-				"WorkspaceMenuStructure"
+				"WorkspaceMenuStructure",
+				"MovieScene"
 			}
 			);
 	}
diff --git a/Engine/Plugins/VirtualProduction/PerformanceCaptureWorkflow/Source/PerformanceCaptureWorkflow/Private/PCapSessionTemplate.h b/Engine/Plugins/VirtualProduction/PerformanceCaptureWorkflow/Source/PerformanceCaptureWorkflow/Private/PCapSessionTemplate.h
--- a/Engine/Plugins/VirtualProduction/PerformanceCaptureWorkflow/Source/PerformanceCaptureWorkflow/Private/PCapSessionTemplate.h
+++ b/Engine/Plugins/VirtualProduction/PerformanceCaptureWorkflow/Source/PerformanceCaptureWorkflow/Private/PCapSessionTemplate.h
@@ -59,11 +59,11 @@ public:
 	* Core session fields
 	*/
 
-	UPROPERTY(EditAnywhere, BlueprintSetter = SetRootFolder, BlueprintGetter = GetRootFolder, Category="Template", meta = (DisplayName = "Root Folder", ContentDir, EditCondition="bIsEditable"))
+	UPROPERTY(EditAnywhere, BlueprintSetter = SetRootFolder, BlueprintGetter = GetRootFolder, Category="Template", meta = (ContentDir, DisplayName = "Root Folder", EditCondition="bIsEditable"))
 	FDirectoryPath TemplateRootFolder;
 
 	/*
-	 *These properties are private because we want to force the use of the Blueprint Getter and Setter so modifications only follow that path.
+	 *These properties are hidden from BP access because we want to force the use of the Blueprint Getter and Setter so modifications only follow that path.
 	 */
 
 	virtual void PostLoad() override;
@@ -77,7 +77,7 @@ public:
 	FString SessionName;
 
 	/** The final name of the session, evaluated from the given tokens and static strings */
-	UPROPERTY(EditAnywhere, BlueprintReadOnly, Category="Template", meta = (DisplayName = "Session Token", ContentDir, EditCondition="bIsEditable", tooltip="The final name of our session, including any tokens you include in it's construction. You can use tokens to build this string, including the Session Name above"))
+	UPROPERTY(EditAnywhere, BlueprintReadOnly, Category="Template", meta = (DisplayName = "Session Token", EditCondition="bIsEditable", tooltip="The final name of our session, including any tokens you include in it's construction. You can use tokens to build this string, including the Session Name above"))
 	FPCapTokenisedString SessionToken;
 
 	/**
@@ -155,27 +155,27 @@ public:
 	FPCapTokenisedFolderPath SessionFolder;
 
 	/**Common folder template**/
-	UPROPERTY(EditAnywhere, BlueprintReadOnly, Category="Folders|Session Subfolders", meta = (DisplayName = "Common Folder", ContentDir, EditCondition="bIsEditable"))
+	UPROPERTY(EditAnywhere, BlueprintReadOnly, Category="Folders|Session Subfolders", meta = (DisplayName = "Common Folder", EditCondition="bIsEditable"))
 	FPCapTokenisedFolderPath CommonFolder;
 
 	/**Character folder template*/	
-	UPROPERTY(EditAnywhere, BlueprintReadOnly, Category="Folders|Session Subfolders", meta = (DisplayName = "Character Folder", ContentDir, EditCondition="bIsEditable"))
+	UPROPERTY(EditAnywhere, BlueprintReadOnly, Category="Folders|Session Subfolders", meta = (DisplayName = "Character Folder", EditCondition="bIsEditable"))
 	FPCapTokenisedFolderPath CharacterFolder;
 
 	/**Performer folder template*/
-	UPROPERTY(EditAnywhere, BlueprintReadOnly, Category="Folders|Session Subfolders", meta = (DisplayName = "Performer Folder", ContentDir, EditCondition="bIsEditable"))
+	UPROPERTY(EditAnywhere, BlueprintReadOnly, Category="Folders|Session Subfolders", meta = (DisplayName = "Performer Folder", EditCondition="bIsEditable"))
 	FPCapTokenisedFolderPath PerformerFolder;
 
 	/**Prop folder template*/
-	UPROPERTY(EditAnywhere, BlueprintReadOnly, Category="Folders|Session Subfolders", meta = (DisplayName = "Prop Folder", ContentDir, EditCondition="bIsEditable"))
+	UPROPERTY(EditAnywhere, BlueprintReadOnly, Category="Folders|Session Subfolders", meta = (DisplayName = "Prop Folder", EditCondition="bIsEditable"))
 	FPCapTokenisedFolderPath PropFolder;
 
 	/**Scene data folder template*/
-	UPROPERTY(EditAnywhere, BlueprintReadOnly, Category="Folders|Session Subfolders", meta = (DisplayName = "Scene Folder", ContentDir, EditCondition="bIsEditable"))
+	UPROPERTY(EditAnywhere, BlueprintReadOnly, Category="Folders|Session Subfolders", meta = (DisplayName = "Scene Folder", EditCondition="bIsEditable"))
 	FPCapTokenisedFolderPath SceneFolder;
 
 	/**Take Recorder folder template*/
-	UPROPERTY(EditAnywhere, BlueprintReadOnly, Category="Folders|Session Subfolders", meta = (DisplayName = "Takes Folder", ContentDir, EditCondition="bIsEditable"))
+	UPROPERTY(EditAnywhere, BlueprintReadOnly, Category="Folders|Session Subfolders", meta = (DisplayName = "Takes Folder", EditCondition="bIsEditable"))
 	FPCapTokenisedFolderPath TakeFolder;
 
 	/** Additional folders. Create any additional folders you want under the session folder. You can use the Map key to label what each folder is for*/
@@ -195,6 +195,10 @@ public:
 	/** Bool to control if recorded sequences should start at the current timecode value. */
 	UPROPERTY(EditAnywhere, BlueprintReadOnly, Category="Take Recorder", meta = (DisplayName = "Start at Current Timecode", EditCondition="bIsEditable"))
 	bool bStartAtCurrentTimecode = true;
+	
+	/** The clock source to use with recordings. Recommended is "Timecode". */
+	UPROPERTY(EditAnywhere, BlueprintReadOnly, Category="Take Recorder", meta = (DisplayName = "Recording Clock Source", EditCondition="bIsEditable"))
+	EUpdateClockSource RecordingClockSource = EUpdateClockSource::Timecode;
 
 	/** Bool to control if each recording source is placed into a subscene at the start of recording. */
 	UPROPERTY(EditAnywhere, BlueprintReadOnly, Category="Take Recorder", meta = (DisplayName = "Record to Subsequences", EditCondition="bIsEditable"))
