Mocap Manager - add default folder for ingested capture data - this will be video ingested using capture manager. Update to viewmodel to handle new folder creation, getting and setting the import folder for capture manager.

#virtualized
#Synced-CL 51493371

diff --git a/Engine/Plugins/VirtualProduction/PerformanceCaptureWorkflow/PerformanceCaptureWorkflow.uplugin b/Engine/Plugins/VirtualProduction/PerformanceCaptureWorkflow/PerformanceCaptureWorkflow.uplugin
--- a/Engine/Plugins/VirtualProduction/PerformanceCaptureWorkflow/PerformanceCaptureWorkflow.uplugin
+++ b/Engine/Plugins/VirtualProduction/PerformanceCaptureWorkflow/PerformanceCaptureWorkflow.uplugin
@@ -82,6 +82,10 @@
 		{
 			"Name": "LiveLinkUnrealDevice",
 			"Enabled": true
+		},
+		{
+			"Name": "CaptureManagerEditor",
+			"Enabled": true
 		}
 	]
 }
\ No newline at end of file
diff --git a/Engine/Plugins/VirtualProduction/PerformanceCaptureWorkflow/Source/PerformanceCaptureWorkflow/Private/PCapDatabase.h b/Engine/Plugins/VirtualProduction/PerformanceCaptureWorkflow/Source/PerformanceCaptureWorkflow/Private/PCapDatabase.h
--- a/Engine/Plugins/VirtualProduction/PerformanceCaptureWorkflow/Source/PerformanceCaptureWorkflow/Private/PCapDatabase.h
+++ b/Engine/Plugins/VirtualProduction/PerformanceCaptureWorkflow/Source/PerformanceCaptureWorkflow/Private/PCapDatabase.h
@@ -196,6 +196,10 @@ public:
 	/** The content browser path to this session record's character data. */
 	UPROPERTY(VisibleAnywhere, BlueprintReadWrite, Category="Performance Capture Session", meta = (DisplayName = "Character Folder", ContentDir))
 	FString CharacterPath;
+	
+	/** The content browser path to this session record's ingest data. */
+	UPROPERTY(VisibleAnywhere, BlueprintReadWrite, Category="Performance Capture Session", meta = (DisplayName = "Ingest Folder", ContentDir))
+	FString IngestPath;
 
 	/** The content browser path to this session record's recorded takes data. */
 	UPROPERTY(VisibleAnywhere, BlueprintReadWrite, Category="Performance Capture Session", meta = (DisplayName = "Takes Folder", ContentDir))
diff --git a/Engine/Plugins/VirtualProduction/PerformanceCaptureWorkflow/Source/PerformanceCaptureWorkflow/Private/PCapSessionTemplate.cpp b/Engine/Plugins/VirtualProduction/PerformanceCaptureWorkflow/Source/PerformanceCaptureWorkflow/Private/PCapSessionTemplate.cpp
--- a/Engine/Plugins/VirtualProduction/PerformanceCaptureWorkflow/Source/PerformanceCaptureWorkflow/Private/PCapSessionTemplate.cpp
+++ b/Engine/Plugins/VirtualProduction/PerformanceCaptureWorkflow/Source/PerformanceCaptureWorkflow/Private/PCapSessionTemplate.cpp
@@ -32,7 +32,7 @@ void UPCapSessionTemplate::PostEditChangeProperty(FPropertyChangedEvent& Propert
 void UPCapSessionTemplate::UpdateAllFields()
 {
 
-	if (bIsEditable) //If this asset is "locked" then do no update any of the strings. 
+	if (bIsEditable) //If this asset is not "locked" then do not update any of the strings. 
 	{
 		//Make sure all the tokens are updated each time a property is edited
 		SessionToken		= UpdateStringTemplate(SessionToken);
@@ -40,6 +40,7 @@ void UPCapSessionTemplate::UpdateAllFields()
 		SessionFolder		= UpdateFolderPathTemplate(SessionFolder);
 		CommonFolder		= UpdateFolderPathTemplate(CommonFolder);
 		CharacterFolder		= UpdateFolderPathTemplate(CharacterFolder);
+		IngestFolder		= UpdateFolderPathTemplate(IngestFolder);
 		PerformerFolder		= UpdateFolderPathTemplate(PerformerFolder);
 		PropFolder			= UpdateFolderPathTemplate(PropFolder);
 		SceneFolder			= UpdateFolderPathTemplate(SceneFolder);
diff --git a/Engine/Plugins/VirtualProduction/PerformanceCaptureWorkflow/Source/PerformanceCaptureWorkflow/Private/PCapSessionTemplate.h b/Engine/Plugins/VirtualProduction/PerformanceCaptureWorkflow/Source/PerformanceCaptureWorkflow/Private/PCapSessionTemplate.h
--- a/Engine/Plugins/VirtualProduction/PerformanceCaptureWorkflow/Source/PerformanceCaptureWorkflow/Private/PCapSessionTemplate.h
+++ b/Engine/Plugins/VirtualProduction/PerformanceCaptureWorkflow/Source/PerformanceCaptureWorkflow/Private/PCapSessionTemplate.h
@@ -161,6 +161,10 @@ public:
 	/**Character folder template*/	
 	UPROPERTY(EditAnywhere, BlueprintReadOnly, Category="Folders|Session Subfolders", meta = (DisplayName = "Character Folder", EditCondition="bIsEditable"))
 	FPCapTokenisedFolderPath CharacterFolder;
+	
+	/**Ingest data folder template*/
+	UPROPERTY(EditAnywhere, BlueprintReadOnly, Category="Folders|Session Subfolders", meta = (DisplayName = "Ingest Folder", EditCondition="bIsEditable"))
+	FPCapTokenisedFolderPath IngestFolder;
 
 	/**Performer folder template*/
 	UPROPERTY(EditAnywhere, BlueprintReadOnly, Category="Folders|Session Subfolders", meta = (DisplayName = "Performer Folder", EditCondition="bIsEditable"))
