[Performance Capture]
- Bump both Pcap plugins to Beta status

[RN] Performance Capture Core and Workflow plugins are now Beta status
#Synced-CL 52621925

diff --git a/Engine/Plugins/Animation/PerformanceCaptureCore/PerformanceCaptureCore.uplugin b/Engine/Plugins/Animation/PerformanceCaptureCore/PerformanceCaptureCore.uplugin
--- a/Engine/Plugins/Animation/PerformanceCaptureCore/PerformanceCaptureCore.uplugin
+++ b/Engine/Plugins/Animation/PerformanceCaptureCore/PerformanceCaptureCore.uplugin
@@ -10,8 +10,8 @@
 	"DocsURL": "",
 	"MarketplaceURL": "",
 	"CanContainContent": false,
-	"IsBetaVersion": false,
-	"IsExperimentalVersion": true,
+	"IsBetaVersion": true,
+	"IsExperimentalVersion": false,
 	"Installed": false,
 	"Modules": [
 		{
diff --git a/Engine/Plugins/VirtualProduction/PerformanceCaptureWorkflow/PerformanceCaptureWorkflow.uplugin b/Engine/Plugins/VirtualProduction/PerformanceCaptureWorkflow/PerformanceCaptureWorkflow.uplugin
--- a/Engine/Plugins/VirtualProduction/PerformanceCaptureWorkflow/PerformanceCaptureWorkflow.uplugin
+++ b/Engine/Plugins/VirtualProduction/PerformanceCaptureWorkflow/PerformanceCaptureWorkflow.uplugin
@@ -11,8 +11,8 @@
 	"MarketplaceURL": "",
 	"SupportURL": "",
 	"CanContainContent": true,
-	"IsBetaVersion": false,
-	"IsExperimentalVersion": true,
+	"IsBetaVersion": true,
+	"IsExperimentalVersion": false,
 	"Installed": false,
 	"Modules": [
 		{
